Moved BSP-related functionality into a Developer module, where it can be used by plugins/modules other than UnrealEd.

diff --git a/Engine/Plugins/Editor/GeometryMode/Source/GeometryMode/GeometryMode.Build.cs b/Engine/Plugins/Editor/GeometryMode/Source/GeometryMode/GeometryMode.Build.cs
--- a/Engine/Plugins/Editor/GeometryMode/Source/GeometryMode/GeometryMode.Build.cs
+++ b/Engine/Plugins/Editor/GeometryMode/Source/GeometryMode/GeometryMode.Build.cs
@@ -8,6 +8,7 @@ public class GeometryMode : ModuleRules
 	{
 		PrivateDependencyModuleNames.AddRange(
 			new string[] {
+				"BSPUtils",
 				"Core",
 				"CoreUObject",
 				"Engine",
diff --git a/Engine/Plugins/Editor/GeometryMode/Source/GeometryMode/Private/GeometryModifiers.cpp b/Engine/Plugins/Editor/GeometryMode/Source/GeometryMode/Private/GeometryModifiers.cpp
--- a/Engine/Plugins/Editor/GeometryMode/Source/GeometryMode/Private/GeometryModifiers.cpp
+++ b/Engine/Plugins/Editor/GeometryMode/Source/GeometryMode/Private/GeometryModifiers.cpp
@@ -18,6 +18,7 @@
 #include "EditorModes.h"
 #include "Dialogs/Dialogs.h"
 #include "BSPOps.h"
+#include "Engine/BrushBuilder.h"
 #include "GeometryEdMode.h"
 #include "EditorGeometry.h"
 #include "Engine/BrushShape.h"
@@ -1351,7 +1352,12 @@ void UGeomModifier_Pen::Apply()
 		}
 		else
 		{
-			ResultingBrush = FBSPOps::csgAddOperation( BuilderBrush, PF_DefaultFlags, Brush_Add );	
+			ResultingBrush = FBSPOps::csgAddOperation( BuilderBrush, PF_DefaultFlags, Brush_Add );
+
+			if (ResultingBrush && ResultingBrush->GetBrushBuilder())
+			{
+				FActorLabelUtilities::SetActorLabelUnique(ResultingBrush, FText::Format(NSLOCTEXT("BSPBrushOps", "BrushName", "{0} Brush"), FText::FromString(ResultingBrush->GetBrushBuilder()->GetClass()->GetDescription())).ToString());
+			}
 		}
 
 		// Make sure the graphics engine isn't busy rendering this geometry before we go and modify it
diff --git a/Engine/Plugins/Experimental/MeshModelingToolsetExp/Source/MeshModelingToolsEditorOnly/MeshModelingToolsEditorOnly.Build.cs b/Engine/Plugins/Experimental/MeshModelingToolsetExp/Source/MeshModelingToolsEditorOnly/MeshModelingToolsEditorOnly.Build.cs
--- a/Engine/Plugins/Experimental/MeshModelingToolsetExp/Source/MeshModelingToolsEditorOnly/MeshModelingToolsEditorOnly.Build.cs
+++ b/Engine/Plugins/Experimental/MeshModelingToolsetExp/Source/MeshModelingToolsEditorOnly/MeshModelingToolsEditorOnly.Build.cs
@@ -49,6 +49,7 @@ public class MeshModelingToolsEditorOnly : ModuleRules
 		PrivateDependencyModuleNames.AddRange(
 			new string[]
 			{
+				"BSPUtils",
 				"CoreUObject",
 				"Engine",
                 "RenderCore",
diff --git a/Engine/Plugins/Experimental/MeshModelingToolsetExp/Source/MeshModelingToolsEditorOnly/Private/MeshToVolumeTool.cpp b/Engine/Plugins/Experimental/MeshModelingToolsetExp/Source/MeshModelingToolsEditorOnly/Private/MeshToVolumeTool.cpp
--- a/Engine/Plugins/Experimental/MeshModelingToolsetExp/Source/MeshModelingToolsEditorOnly/Private/MeshToVolumeTool.cpp
+++ b/Engine/Plugins/Experimental/MeshModelingToolsetExp/Source/MeshModelingToolsEditorOnly/Private/MeshToVolumeTool.cpp
@@ -20,7 +20,7 @@
 #include "Engine/Classes/Components/BrushComponent.h"
 #include "Engine/Classes/Engine/Polys.h"
 #include "Model.h"
-#include "BSPOps.h"		// in UnrealEd
+#include "BSPOps.h"
 
 #include "ExplicitUseGeometryMathTypes.h"		// using UE::Geometry::(math types)
 using namespace UE::Geometry;
diff --git a/Engine/Plugins/Experimental/MeshModelingToolsetExp/Source/ModelingComponentsEditorOnly/ModelingComponentsEditorOnly.Build.cs b/Engine/Plugins/Experimental/MeshModelingToolsetExp/Source/ModelingComponentsEditorOnly/ModelingComponentsEditorOnly.Build.cs
--- a/Engine/Plugins/Experimental/MeshModelingToolsetExp/Source/ModelingComponentsEditorOnly/ModelingComponentsEditorOnly.Build.cs
+++ b/Engine/Plugins/Experimental/MeshModelingToolsetExp/Source/ModelingComponentsEditorOnly/ModelingComponentsEditorOnly.Build.cs
@@ -44,6 +44,7 @@ public class ModelingComponentsEditorOnly : ModuleRules
 		PrivateDependencyModuleNames.AddRange(
 			new string[]
 			{
+				"BSPUtils",
 				"CoreUObject",
 				"Engine",
 				"InputCore",
diff --git a/Engine/Source/Developer/BSPUtils/BSPUtils.Build.cs b/Engine/Source/Developer/BSPUtils/BSPUtils.Build.cs
new file mode 100644
--- /dev/null
+++ b/Engine/Source/Developer/BSPUtils/BSPUtils.Build.cs
@@ -0,0 +1,25 @@
+// Copyright Epic Games, Inc. All Rights Reserved.
+
+namespace UnrealBuildTool.Rules
+{
+    public class BSPUtils : ModuleRules
+    {
+        public BSPUtils(ReadOnlyTargetRules Target) : base(Target)
+        {
+            PrivateDependencyModuleNames.AddRange(
+                new string[]
+                {
+					"Core",
+					"CoreUObject",
+					"Engine"
+				}
+			);
+
+			PublicDependencyModuleNames.AddRange(
+				new string[]
+				{
+				}
+			);
+        }
+    }
+}
diff --git a/Engine/Source/Developer/BSPUtils/Private/BSPOps.cpp b/Engine/Source/Developer/BSPUtils/Private/BSPOps.cpp
new file mode 100644
--- /dev/null
+++ b/Engine/Source/Developer/BSPUtils/Private/BSPOps.cpp
@@ -0,0 +1,1442 @@
+// Copyright Epic Games, Inc. All Rights Reserved.
+
+
+#include "BSPOps.h"
+#include "EngineDefines.h"
+#include "Model.h"
+#include "Materials/Material.h"
+#include "Engine/BrushBuilder.h"
+#include "Editor/EditorEngine.h"
+#include "Components/BrushComponent.h"
+#include "GameFramework/Volume.h"
+
+DEFINE_LOG_CATEGORY_STATIC(LogBSPOps, Log, All);
+
+/** Errors encountered in Csg operation. */
+int32 FBSPOps::GErrors = 0;
+bool FBSPOps::GFastRebuild = false;
+FBspPointsGrid* FBspPointsGrid::GBspPoints = nullptr;
+FBspPointsGrid* FBspPointsGrid::GBspVectors = nullptr;
+
+static void TagReferencedNodes( UModel *Model, int32 *NodeRef, int32 *PolyRef, int32 iNode )
+{
+	FBspNode &Node = Model->Nodes[iNode];
+
+	NodeRef[iNode     ] = 0;
+	PolyRef[Node.iSurf] = 0;
+
+	if( Node.iFront != INDEX_NONE ) TagReferencedNodes(Model,NodeRef,PolyRef,Node.iFront);
+	if( Node.iBack  != INDEX_NONE ) TagReferencedNodes(Model,NodeRef,PolyRef,Node.iBack );
+	if( Node.iPlane != INDEX_NONE ) TagReferencedNodes(Model,NodeRef,PolyRef,Node.iPlane);
+}
+
+//
+// Update a bounding volume by expanding it to enclose a list of polys.
+//
+static void UpdateBoundWithPolys( FBox& Bound, FPoly** PolyList, int32 nPolys )
+{
+	for( int32 i=0; i<nPolys; i++ )
+		for( int32 j=0; j<PolyList[i]->Vertices.Num(); j++ )
+			Bound += PolyList[i]->Vertices[j];
+}
+
+//
+// Update a convolution hull with a list of polys.
+//
+static void UpdateConvolutionWithPolys( UModel *Model, int32 iNode, FPoly **PolyList, int32 nPolys )
+{
+	FBox Box(ForceInit);
+
+	FBspNode &Node = Model->Nodes[iNode];
+	Node.iCollisionBound = Model->LeafHulls.Num();
+	for( int32 i=0; i<nPolys; i++ )
+	{
+		if( PolyList[i]->iBrushPoly != INDEX_NONE )
+		{
+			int32 j;
+			for( j=0; j<i; j++ )
+				if( PolyList[j]->iBrushPoly == PolyList[i]->iBrushPoly )
+					break;
+			if( j >= i )
+				Model->LeafHulls.Add(PolyList[i]->iBrushPoly);
+		}
+		for( int32 j=0; j<PolyList[i]->Vertices.Num(); j++ )
+			Box += PolyList[i]->Vertices[j];
+	}
+	Model->LeafHulls.Add(INDEX_NONE);
+
+	// Add bounds.
+	Model->LeafHulls.Add( *(int32*)&Box.Min.X );
+	Model->LeafHulls.Add( *(int32*)&Box.Min.Y );
+	Model->LeafHulls.Add( *(int32*)&Box.Min.Z );
+	Model->LeafHulls.Add( *(int32*)&Box.Max.X );
+	Model->LeafHulls.Add( *(int32*)&Box.Max.Y );
+	Model->LeafHulls.Add( *(int32*)&Box.Max.Z );
+
+}
+
+//
+// Cut a partitioning poly by a list of polys, and add the resulting inside pieces to the
+// front list and back list.
+//
+static void SplitPartitioner
+(
+	UModel*	Model,
+	FPoly**	PolyList,
+	FPoly**	FrontList,
+	FPoly**	BackList,
+	int32		n,
+	int32		nPolys,
+	int32&	nFront, 
+	int32&	nBack, 
+	FPoly	InfiniteEdPoly,
+	TArray<FPoly*>& AllocatedFPolys
+)
+{
+	FPoly FrontPoly,BackPoly;
+	while( n < nPolys )
+	{
+		FPoly* Poly = PolyList[n];
+		switch( InfiniteEdPoly.SplitWithPlane(Poly->Vertices[0],Poly->Normal,&FrontPoly,&BackPoly,0) )
+		{
+			case SP_Coplanar:
+				// May occasionally happen.
+//				UE_LOG(LogBSPOps, Log,  TEXT("FilterBound: Got inficoplanar") );
+				break;
+
+			case SP_Front:
+				// Shouldn't happen if hull is correct.
+//				UE_LOG(LogBSPOps, Log,  TEXT("FilterBound: Got infifront") );
+				return;
+
+			case SP_Split:
+				InfiniteEdPoly = BackPoly;
+				break;
+
+			case SP_Back:
+				break;
+		}
+		n++;
+	}
+
+	FPoly* New = new FPoly;
+	*New = InfiniteEdPoly;
+	New->Reverse();
+	New->iBrushPoly |= 0x40000000;
+	FrontList[nFront++] = New;
+	AllocatedFPolys.Add( New );
+	
+	New = new FPoly;
+	*New = InfiniteEdPoly;
+	BackList[nBack++] = New;
+	AllocatedFPolys.Add( New );
+}
+
+//
+// Build an FPoly representing an "infinite" plane (which exceeds the maximum
+// dimensions of the world in all directions) for a particular Bsp node.
+//
+FPoly FBSPOps::BuildInfiniteFPoly( UModel* Model, int32 iNode )
+{
+	FBspNode &Node   = Model->Nodes  [iNode       ];
+	FBspSurf &Poly   = Model->Surfs  [Node.iSurf  ];
+	FVector  Base    = Poly.Plane * Poly.Plane.W;
+	FVector  Normal  = Poly.Plane;
+	FVector	 Axis1,Axis2;
+
+	// Find two non-problematic axis vectors.
+	Normal.FindBestAxisVectors( Axis1, Axis2 );
+
+	// Set up the FPoly.
+	FPoly EdPoly;
+	EdPoly.Init();
+	EdPoly.Normal      = Normal;
+	EdPoly.Base        = Base;
+	new(EdPoly.Vertices) FVector3f(Base + Axis1*WORLD_MAX + Axis2*WORLD_MAX);
+	new(EdPoly.Vertices) FVector3f(Base - Axis1*WORLD_MAX + Axis2*WORLD_MAX);
+	new(EdPoly.Vertices) FVector3f(Base - Axis1*WORLD_MAX - Axis2*WORLD_MAX);
+	new(EdPoly.Vertices) FVector3f(Base + Axis1*WORLD_MAX - Axis2*WORLD_MAX);
+
+	return EdPoly;
+}
+
+//
+// Recursively filter a set of polys defining a convex hull down the Bsp,
+// splitting it into two halves at each node and adding in the appropriate
+// face polys at splits.
+//
+static void FilterBound
+(
+	UModel*			Model,
+	FBox*			ParentBound,
+	int32				iNode,
+	FPoly**			PolyList,
+	int32				nPolys,
+	int32				Outside
+)
+{
+	FMemMark Mark(FMemStack::Get());
+	FBspNode&	Node	= Model->Nodes  [iNode];
+	FBspSurf&	Surf	= Model->Surfs  [Node.iSurf];
+	FVector3f	Base = Surf.Plane * Surf.Plane.W;
+	FVector3f&	Normal	= Model->Vectors[Surf.vNormal];
+	FBox		Bound(ForceInit);
+
+	Bound.Min.X = Bound.Min.Y = Bound.Min.Z = +WORLD_MAX;
+	Bound.Max.X = Bound.Max.Y = Bound.Max.Z = -WORLD_MAX;
+
+	// Split bound into front half and back half.
+	FPoly** FrontList = new(FMemStack::Get(),nPolys*2+16)FPoly*; int32 nFront=0;
+	FPoly** BackList  = new(FMemStack::Get(),nPolys*2+16)FPoly*; int32 nBack=0;
+
+	// Keeping track of allocated FPoly structures to delete later on.
+	TArray<FPoly*> AllocatedFPolys;
+
+	FPoly* FrontPoly  = new FPoly;
+	FPoly* BackPoly   = new FPoly;
+
+	// Keep track of allocations.
+	AllocatedFPolys.Add( FrontPoly );
+	AllocatedFPolys.Add( BackPoly );
+
+	for( int32 i=0; i<nPolys; i++ )
+	{
+		FPoly *Poly = PolyList[i];
+		switch( Poly->SplitWithPlane( Base, Normal, FrontPoly, BackPoly, 0 ) )
+		{
+			case SP_Coplanar:
+//				UE_LOG(LogBSPOps, Log,  TEXT("FilterBound: Got coplanar") );
+				FrontList[nFront++] = Poly;
+				BackList[nBack++] = Poly;
+				break;
+			
+			case SP_Front:
+				FrontList[nFront++] = Poly;
+				break;
+			
+			case SP_Back:
+				BackList[nBack++] = Poly;
+				break;
+			
+			case SP_Split:
+				FrontList[nFront++] = FrontPoly;
+				BackList [nBack++] = BackPoly;
+
+				FrontPoly = new FPoly;
+				BackPoly  = new FPoly;
+
+				// Keep track of allocations.
+				AllocatedFPolys.Add( FrontPoly );
+				AllocatedFPolys.Add( BackPoly );
+
+				break;
+
+			default:
+				UE_LOG(LogBSPOps, Fatal, TEXT("FZoneFilter::FilterToLeaf: Unknown split code") );
+		}
+	}
+	if( nFront && nBack )
+	{
+		// Add partitioner plane to front and back.
+		FPoly InfiniteEdPoly = FBSPOps::BuildInfiniteFPoly( Model, iNode );
+		InfiniteEdPoly.iBrushPoly = iNode;
+
+		SplitPartitioner(Model,PolyList,FrontList,BackList,0,nPolys,nFront,nBack,InfiniteEdPoly,AllocatedFPolys);
+	}
+	else
+	{
+// 		if( !nFront ) UE_LOG(LogBSPOps, Log,  TEXT("FilterBound: Empty fronthull") );
+// 		if( !nBack  ) UE_LOG(LogBSPOps, Log,  TEXT("FilterBound: Empty backhull") );
+	}
+
+	// Recursively update all our childrens' bounding volumes.
+	if( nFront > 0 )
+	{
+		if( Node.iFront != INDEX_NONE )
+			FilterBound( Model, &Bound, Node.iFront, FrontList, nFront, Outside || Node.IsCsg() );
+		else if( Outside || Node.IsCsg() )
+			UpdateBoundWithPolys( Bound, FrontList, nFront );
+		else
+			UpdateConvolutionWithPolys( Model, iNode, FrontList, nFront );
+	}
+	if( nBack > 0 )
+	{
+		if( Node.iBack != INDEX_NONE)
+			FilterBound( Model, &Bound,Node.iBack, BackList, nBack, Outside && !Node.IsCsg() );
+		else if( Outside && !Node.IsCsg() )
+			UpdateBoundWithPolys( Bound, BackList, nBack );
+		else
+			UpdateConvolutionWithPolys( Model, iNode, BackList, nBack );
+	}
+
+	// Update parent bound to enclose this bound.
+	if( ParentBound )
+		*ParentBound += Bound;
+
+	// Delete FPolys allocated above. We cannot use FMemStack::Get() for FPoly as the array data FPoly contains will be allocated in regular memory.
+	for( int32 i=0; i<AllocatedFPolys.Num(); i++ )
+	{
+		FPoly* AllocatedFPoly = AllocatedFPolys[i];
+		delete AllocatedFPoly;
+	}
+
+	Mark.Pop();
+}
+
+/*-----------------------------------------------------------------------------
+	Bsp Splitting.
+-----------------------------------------------------------------------------*/
+
+//
+// Find the best splitting polygon within a pool of polygons, and return its
+// index (into the PolyList array).
+//
+static FPoly *FindBestSplit
+(
+	int32					NumPolys,
+	FPoly**				PolyList,
+	FBSPOps::EBspOptimization	Opt,
+	int32					Balance,
+	int32					InPortalBias
+)
+{
+	check(NumPolys>0);
+
+	// No need to test if only one poly.
+	if( NumPolys==1 )
+		return PolyList[0];
+
+	FPoly   *Poly, *Best=NULL;
+	float   Score, BestScore;
+	int32     i, Index, j, Inc;
+	int32     Splits, Front, Back, Coplanar, AllSemiSolids;
+
+	//PortalBias -- added by Legend on 4/12/2000
+	float	PortalBias = InPortalBias / 100.0f;
+	Balance &= 0xFF;								// keep only the low byte to recover "Balance"
+	//UE_LOG(LogBSPOps, Log, TEXT("Balance=%d PortalBias=%f"), Balance, PortalBias );
+
+	if		(Opt==FBSPOps::BSP_Optimal)  Inc = 1;					// Test lots of nodes.
+	else if (Opt==FBSPOps::BSP_Good)		Inc = FMath::Max(1,NumPolys/20);	// Test 20 nodes.
+	else /* BSP_Lame */			Inc = FMath::Max(1,NumPolys/4);	// Test 4 nodes.
+
+	// See if there are any non-semisolid polygons here.
+	for( i=0; i<NumPolys; i++ )
+		if( !(PolyList[i]->PolyFlags & PF_AddLast) )
+			break;
+	AllSemiSolids = (i>=NumPolys);
+
+	// Search through all polygons in the pool and find:
+	// A. The number of splits each poly would make.
+	// B. The number of front and back nodes the polygon would create.
+	// C. Number of coplanars.
+	BestScore = 0;
+	for( i=0; i<NumPolys; i+=Inc )
+	{
+		Splits = Front = Back = Coplanar = 0;
+		Index = i-1;
+		do
+		{
+			Index++;
+			Poly = PolyList[Index];
+		} while( Index<(i+Inc) && Index<NumPolys 
+			&& ( (Poly->PolyFlags & PF_AddLast) && !(Poly->PolyFlags & PF_Portal) )
+			&& !AllSemiSolids );
+		if( Index>=i+Inc || Index>=NumPolys )
+			continue;
+
+		for( j=0; j<NumPolys; j+=Inc ) if( j != Index )
+		{
+			FPoly *OtherPoly = PolyList[j];
+			switch( OtherPoly->SplitWithPlaneFast( FPlane( Poly->Vertices[0], Poly->Normal), NULL, NULL ) )
+			{
+				case SP_Coplanar:
+					Coplanar++;
+					break;
+
+				case SP_Front:
+					Front++;
+					break;
+
+				case SP_Back:
+					Back++;
+					break;
+
+				case SP_Split:
+					// Disfavor splitting polys that are zone portals.
+					if( !(OtherPoly->PolyFlags & PF_Portal) )
+						Splits++;
+					else
+						Splits += 16;
+					break;
+			}
+		}
+		// added by Legend 1/31/1999
+		// Score optimization: minimize cuts vs. balance tree (as specified in BSP Rebuilder dialog)
+		Score = ( 100.0 - float(Balance) ) * Splits + float(Balance) * FMath::Abs( Front - Back );
+		if( Poly->PolyFlags & PF_Portal )
+		{
+			// PortalBias -- added by Legend on 4/12/2000
+			//
+			// PortalBias enables level designers to control the effect of Portals on the BSP.
+			// This effect can range from 0.0 (ignore portals), to 1.0 (portals cut everything).
+			//
+			// In builds prior to this (since the 221 build dating back to 1/31/1999) the bias
+			// has been 1.0 causing the portals to cut the BSP in ways that will potentially
+			// degrade level performance, and increase the BSP complexity.
+			// 
+			// By setting the bias to a value between 0.3 and 0.7 the positive effects of 
+			// the portals are preserved without giving them unreasonable priority in the BSP.
+			//
+			// Portals should be weighted high enough in the BSP to separate major parts of the
+			// level from each other (pushing entire rooms down the branches of the BSP), but
+			// should not be so high that portals cut through adjacent geometry in a way that
+			// increases complexity of the room being (typically, accidentally) cut.
+			//
+			Score -= ( 100.0 - float(Balance) ) * Splits * PortalBias; // ignore PortalBias of the split polys -- bias toward portal selection for cutting planes!
+		}
+		//UE_LOG(LogBSPOps, Log,  "  %4d: Score = %f (Front = %4d, Back = %4d, Splits = %4d, Flags = %08X)", Index, Score, Front, Back, Splits, Poly->PolyFlags ); //LEC
+
+		if( Score<BestScore || !Best )
+		{
+			Best      = Poly;
+			BestScore = Score;
+		}
+	}
+	check(Best);
+	return Best;
+}
+
+//
+// Pick a splitter poly then split a pool of polygons into front and back polygons and
+// recurse.
+//
+// iParent = Parent Bsp node, or INDEX_NONE if this is the root node.
+// IsFront = 1 if this is the front node of iParent, 0 of back (undefined if iParent==INDEX_NONE)
+//
+void FBSPOps::SplitPolyList
+(
+	UModel				*Model,
+	int32                 iParent,
+	ENodePlace			NodePlace,
+	int32                 NumPolys,
+	FPoly				**PolyList,
+	EBspOptimization	Opt,
+	int32					Balance,
+	int32					PortalBias,
+	int32					RebuildSimplePolys
+)
+{
+	FMemMark Mark(FMemStack::Get());
+
+	// Keeping track of allocated FPoly structures to delete later on.
+	TArray<FPoly*> AllocatedFPolys;
+
+	// To account for big EdPolys split up.
+	int32 NumPolysToAlloc = NumPolys + 8 + NumPolys/4;
+	int32 NumFront=0; FPoly **FrontList = new(FMemStack::Get(),NumPolysToAlloc)FPoly*;
+	int32 NumBack =0; FPoly **BackList  = new(FMemStack::Get(),NumPolysToAlloc)FPoly*;
+
+	FPoly *SplitPoly = FindBestSplit( NumPolys, PolyList, Opt, Balance, PortalBias );
+
+	// Add the splitter poly to the Bsp with either a new BspSurf or an existing one.
+	if( RebuildSimplePolys )
+	{
+		SplitPoly->iLinkSurf = Model->Surfs.Num();
+	}
+
+	int32 iOurNode	= bspAddNode(Model,iParent,NodePlace,0,SplitPoly);
+	int32 iPlaneNode	= iOurNode;
+
+	// Now divide all polygons in the pool into (A) polygons that are
+	// in front of Poly, and (B) polygons that are in back of Poly.
+	// Coplanar polys are inserted immediately, before recursing.
+
+	// If any polygons are split by Poly, we ignrore the original poly,
+	// split it into two polys, and add two new polys to the pool.
+	FPoly *FrontEdPoly = new FPoly;
+	FPoly *BackEdPoly  = new FPoly;
+	// Keep track of allocations.
+	AllocatedFPolys.Add( FrontEdPoly );
+	AllocatedFPolys.Add( BackEdPoly );
+
+	for( int32 i=0; i<NumPolys; i++ )
+	{
+		FPoly *EdPoly = PolyList[i];
+		if( EdPoly == SplitPoly )
+		{
+			continue;
+		}
+
+		switch( EdPoly->SplitWithPlane( SplitPoly->Vertices[0], SplitPoly->Normal, FrontEdPoly, BackEdPoly, 0 ) )
+		{
+			case SP_Coplanar:
+	            if( RebuildSimplePolys )
+				{
+					EdPoly->iLinkSurf = Model->Surfs.Num()-1;
+				}
+				iPlaneNode = bspAddNode( Model, iPlaneNode, NODE_Plane, 0, EdPoly );
+				break;
+			
+			case SP_Front:
+	            FrontList[NumFront++] = PolyList[i];
+				break;
+			
+			case SP_Back:
+	            BackList[NumBack++] = PolyList[i];
+				break;
+			
+			case SP_Split:
+
+				// Create front & back nodes.
+				FrontList[NumFront++] = FrontEdPoly;
+				BackList [NumBack ++] = BackEdPoly;
+
+				FrontEdPoly = new FPoly;
+				BackEdPoly  = new FPoly;
+				// Keep track of allocations.
+				AllocatedFPolys.Add( FrontEdPoly );
+				AllocatedFPolys.Add( BackEdPoly );
+
+				break;
+		}
+	}
+
+	// Recursively split the front and back pools.
+	if( NumFront > 0 ) SplitPolyList( Model, iOurNode, NODE_Front, NumFront, FrontList, Opt, Balance, PortalBias, RebuildSimplePolys );
+	if( NumBack  > 0 ) SplitPolyList( Model, iOurNode, NODE_Back,  NumBack,  BackList,  Opt, Balance, PortalBias, RebuildSimplePolys );
+
+	// Delete FPolys allocated above. We cannot use FMemStack::Get() for FPoly as the array data FPoly contains will be allocated in regular memory.
+	for( int32 i=0; i<AllocatedFPolys.Num(); i++ )
+	{
+		FPoly* AllocatedFPoly = AllocatedFPolys[i];
+		delete AllocatedFPoly;
+	}
+
+	Mark.Pop();
+}
+
+/** Prepare a moving brush. */
+void FBSPOps::csgPrepMovingBrush( ABrush* Actor )
+{
+	check(Actor);
+//	UE_LOG(LogBSPOps, Log,  TEXT("Preparing brush %s"), *Actor->GetName() );  // moved here so that we can easily debug when an actor has lost parts of its brush
+
+	check(Actor->GetBrushComponent());
+	check(Actor->Brush);
+	check(Actor->Brush->RootOutside);
+
+	RebuildBrush(Actor->Brush);
+
+	// Make sure simplified collision is up to date.
+	Actor->GetBrushComponent()->BuildSimpleBrushCollision();
+	Actor->RebuildNavigationData();
+}
+
+/**
+ * Duplicates the specified brush and makes it into a CSG-able level brush.
+ * @return		The new brush, or NULL if the original was empty.
+ */
+void FBSPOps::csgCopyBrush( ABrush* Dest, ABrush* Src, uint32 PolyFlags, EObjectFlags ResFlags, bool bNeedsPrep, bool bCopyPosRotScale, bool bAllowEmpty )
+{
+	check(Src);
+	check(Src->GetBrushComponent());
+	check(Src->Brush);
+
+	// Handle empty brush.
+	if( !bAllowEmpty && !Src->Brush->Polys->Element.Num() )
+	{
+		Dest->Brush = NULL;
+		Dest->GetBrushComponent()->Brush = NULL;
+		return;
+	}
+
+	// Duplicate the brush and its polys.
+	Dest->PolyFlags		= PolyFlags;
+	Dest->Brush = NewObject<UModel>(Dest, NAME_None, ResFlags);
+	Dest->Brush->Initialize(nullptr, Src->Brush->RootOutside);
+	Dest->Brush->Polys	= NewObject<UPolys>(Dest->Brush, NAME_None, ResFlags);
+	Dest->Brush->Polys->Element = Src->Brush->Polys->Element;
+	Dest->GetBrushComponent()->Brush = Dest->Brush;
+	if(Src->BrushBuilder != nullptr)
+	{
+		Dest->BrushBuilder = DuplicateObject<UBrushBuilder>(Src->BrushBuilder, Dest);
+	}
+
+	// Update poly textures.
+	for( int32 i=0; i<Dest->Brush->Polys->Element.Num(); i++ )
+	{
+		Dest->Brush->Polys->Element[i].iBrushPoly = INDEX_NONE;
+	}
+
+	// Copy positioning, and build bounding box.
+	if(bCopyPosRotScale)
+	{
+		Dest->CopyPosRotScaleFrom( Src );
+	}
+
+	// If it's a moving brush, prep it.
+	if( bNeedsPrep )
+	{
+		csgPrepMovingBrush( Dest );
+	}
+}
+
+/**
+ * Adds a brush to the list of CSG brushes in the level, using a CSG operation.
+ *
+ * @return		A newly-created copy of the brush.
+ */
+ABrush*	FBSPOps::csgAddOperation( ABrush* Actor, uint32 PolyFlags, EBrushType BrushType)
+{
+	check(Actor);
+	check(Actor->GetBrushComponent());
+	check(Actor->Brush);
+	check(Actor->Brush->Polys);
+	check(Actor->GetWorld());
+
+	// Can't do this if brush has no polys. 
+	if( !Actor->Brush->Polys->Element.Num() )
+		return NULL;
+
+	// Spawn a new actor for the brush.
+
+	ABrush* Result  = Actor->GetWorld()->SpawnBrush();
+	Result->SetNotForClientOrServer();
+
+	// Duplicate the brush.
+	csgCopyBrush
+	(
+		Result,
+		Actor,
+		PolyFlags,
+		RF_Transactional,
+		0,
+		true
+	);
+	check(Result->Brush);
+
+	// Assign the default material to the brush's polys.
+	for( int32 i=0; i<Result->Brush->Polys->Element.Num(); i++ )
+	{
+		FPoly& CurrentPoly = Result->Brush->Polys->Element[i];
+		if ( !CurrentPoly.Material )
+		{
+			CurrentPoly.Material = UMaterial::GetDefaultMaterial(MD_Surface);
+		}
+	}
+
+	// Set add-info.
+	Result->BrushType = BrushType;
+
+	Result->ReregisterAllComponents();
+
+	return Result;
+}
+
+/** Add a new point to the model (preventing duplicates) and return its index. */
+static int32 AddThing( TArray<FVector3f>& Vectors, const FVector3f& V, float Thresh, int32 Check )
+{
+	if( Check )
+	{
+		// See if this is very close to an existing point/vector.		
+		for( int32 i=0; i<Vectors.Num(); i++ )
+		{
+			const FVector3f &TableVect = Vectors[i];
+			float Temp=(V.X - TableVect.X);
+			if( (Temp > -Thresh) && (Temp < Thresh) )
+			{
+				Temp=(V.Y - TableVect.Y);
+				if( (Temp > -Thresh) && (Temp < Thresh) )
+				{
+					Temp=(V.Z - TableVect.Z);
+					if( (Temp > -Thresh) && (Temp < Thresh) )
+					{
+						// Found nearly-matching vector.
+						return i;
+					}
+				}
+			}
+		}
+	}
+	return Vectors.Add( V );
+}
+
+/** Add a new vector to the model, merging near-duplicates, and return its index. */
+int32 FBSPOps::bspAddVector( UModel* Model, const FVector* V, bool Exact )
+{
+	const float Thresh = Exact ? THRESH_NORMALS_ARE_SAME : THRESH_VECTORS_ARE_NEAR;
+
+	if (FBspPointsGrid::GBspVectors)
+	{
+		// If a points grid has been built for quick vector lookup, use that instead of doing a linear search
+		const int32 NextIndex = Model->Vectors.Num();
+		const int32 ReturnedIndex = FBspPointsGrid::GBspVectors->FindOrAddPoint(*V, NextIndex, Thresh);
+		if (ReturnedIndex == NextIndex)
+		{
+			Model->Vectors.Add(*V);
+		}
+
+		return ReturnedIndex;
+	}
+
+	return AddThing
+	(
+		Model->Vectors,
+		*V,
+		Exact ? THRESH_NORMALS_ARE_SAME : THRESH_VECTORS_ARE_NEAR,
+		1
+	);
+}
+
+/** Add a new point to the model, merging near-duplicates, and return its index. */
+int32 FBSPOps::bspAddPoint( UModel* Model, const FVector* V, bool Exact )
+{
+	const float Thresh = Exact ? THRESH_POINTS_ARE_SAME : THRESH_POINTS_ARE_NEAR;
+
+	if (FBspPointsGrid::GBspPoints)
+	{
+		// If a points grid has been built for quick point lookup, use that instead of doing a linear search
+		const int32 NextIndex = Model->Points.Num();
+		// Always look for points with a low threshold; a generous threshold can result in 'leaks' in the BSP and unwanted polys being generated
+		const int32 ReturnedIndex = FBspPointsGrid::GBspPoints->FindOrAddPoint(*V, NextIndex, THRESH_POINTS_ARE_SAME);
+		if (ReturnedIndex == NextIndex)
+		{
+			Model->Points.Add(*V);
+		}
+
+		return ReturnedIndex;
+	}
+
+	// Try to find a match quickly from the Bsp. This finds all potential matches
+	// except for any dissociated from nodes/surfaces during a rebuild.
+	FVector3f Temp;
+	int32 pVertex;
+	float NearestDist = Model->FindNearestVertex(*V,Temp,Thresh,pVertex);
+	if( (NearestDist >= 0.0) && (NearestDist <= Thresh) )
+	{
+		// Found an existing point.
+		return pVertex;
+	}
+	else
+	{
+		// No match found; add it slowly to find duplicates.
+		return AddThing(Model->Points, *V, Thresh, !GFastRebuild);
+	}
+}
+
+
+/**
+ * Builds Bsp from the editor polygon set (EdPolys) of a model.
+ *
+ * Opt     = Bsp optimization, BSP_Lame (fast), BSP_Good (medium), BSP_Optimal (slow)
+ * Balance = 0-100, 0=only worry about minimizing splits, 100=only balance tree.
+ */
+void FBSPOps::bspBuild( UModel* Model, enum EBspOptimization Opt, int32 Balance, int32 PortalBias, int32 RebuildSimplePolys, int32 iNode )
+{
+	int32 OriginalPolys = Model->Polys->Element.Num();
+
+	// Empty the model's tables.
+	if( RebuildSimplePolys==1 )
+	{
+		// Empty everything but polys.
+		Model->EmptyModel( 1, 0 );
+	}
+	else if( RebuildSimplePolys==0 )
+	{
+		// Empty node vertices.
+		for( int32 i=0; i<Model->Nodes.Num(); i++ )
+			Model->Nodes[i].NumVertices = 0;
+
+		// Refresh the Bsp.
+		bspRefresh(Model,1);
+		
+		// Empty nodes.
+		Model->EmptyModel( 0, 0 );
+	}
+	if( Model->Polys->Element.Num() )
+	{
+		// Allocate polygon pool.
+		FMemMark Mark(FMemStack::Get());
+		FPoly** PolyList = new( FMemStack::Get(), Model->Polys->Element.Num() )FPoly*;
+
+		// Add all FPolys to active list.
+		for( int32 i=0; i<Model->Polys->Element.Num(); i++ )
+			if( Model->Polys->Element[i].Vertices.Num() )
+				PolyList[i] = &Model->Polys->Element[i];
+
+		// Now split the entire Bsp by splitting the list of all polygons.
+		SplitPolyList
+		(
+			Model,
+			INDEX_NONE,
+			NODE_Root,
+			Model->Polys->Element.Num(),
+			PolyList,
+			Opt,
+			Balance,
+			PortalBias,
+			RebuildSimplePolys
+		);
+
+		// Now build the bounding boxes for all nodes.
+		if( RebuildSimplePolys==0 )
+		{
+			// Remove unreferenced things.
+			bspRefresh( Model, 1 );
+
+			// Rebuild all bounding boxes.
+			bspBuildBounds( Model );
+		}
+
+		Mark.Pop();
+	}
+
+//	UE_LOG(LogBSPOps, Log,  TEXT("bspBuild built %i convex polys into %i nodes"), OriginalPolys, Model->Nodes.Num() );
+}
+
+/**
+ * If the Bsp's point and vector tables are nearly full, reorder them and delete unused ones.
+ */
+void FBSPOps::bspRefresh( UModel* Model, bool NoRemapSurfs )
+{
+	FMemStack& MemStack = FMemStack::Get();
+
+	FMemMark Mark(MemStack);
+
+	int32 NumNodes = Model->Nodes.Num();
+	int32 NumSurfs = Model->Surfs.Num();
+	int32 NumVectors = Model->Vectors.Num();
+	int32 NumPoints = Model->Points.Num();
+
+	// Remove unreferenced Bsp surfs.
+	int32* PolyRef;
+	if( NoRemapSurfs )
+	{
+		PolyRef = NewZeroed<int32>(MemStack, NumSurfs);
+	}
+	else
+	{
+		PolyRef = NewOned<int32>(MemStack, NumSurfs);
+	}
+
+	int32* NodeRef = NewOned<int32>(MemStack, NumNodes);
+	if( NumNodes > 0 )
+	{
+		TagReferencedNodes( Model, NodeRef, PolyRef, 0 );
+	}
+
+	// Remap Bsp surfs.
+	{
+		int32 n=0;
+		for( int32 i=0; i<NumSurfs; i++ )
+		{
+			if( PolyRef[i]!=INDEX_NONE )
+			{
+				Model->Surfs[n] = Model->Surfs[i];
+				PolyRef[i]=n++;
+			}
+		}
+		//UE_LOG(LogBSPOps, Log,  TEXT("Polys: %i -> %i"), NumSurfs, n );
+		Model->Surfs.RemoveAt( n, NumSurfs-n );
+		NumSurfs = n;
+	}
+
+	// Remap Bsp nodes.
+	{
+		int32 n=0;
+		for( int32 i=0; i<NumNodes; i++ )
+		{
+			if( NodeRef[i]!=INDEX_NONE )
+			{
+				Model->Nodes[n] = Model->Nodes[i];
+				NodeRef[i]=n++;
+			}
+		}
+		//UE_LOG(LogBSPOps, Log,  TEXT("Nodes: %i -> %i"), NumNodes, n );
+		Model->Nodes.RemoveAt( n, NumNodes-n );
+		NumNodes = n;
+	}
+
+	// Update Bsp nodes.
+	for( int32 i=0; i<NumNodes; i++ )
+	{
+		FBspNode *Node = &Model->Nodes[i];
+		Node->iSurf = PolyRef[Node->iSurf];
+		if (Node->iFront != INDEX_NONE) Node->iFront = NodeRef[Node->iFront];
+		if (Node->iBack  != INDEX_NONE) Node->iBack  = NodeRef[Node->iBack];
+		if (Node->iPlane != INDEX_NONE) Node->iPlane = NodeRef[Node->iPlane];
+	}
+
+	// Remove unreferenced points and vectors.
+	int32* VectorRef = NewOned<int32>(MemStack, NumVectors);
+	int32* PointRef  = NewOned<int32>(MemStack, NumPoints);
+
+	// Check Bsp surfs.
+	TArray<int32*> VertexRef;
+	for( int32 i=0; i<NumSurfs; i++ )
+	{
+		FBspSurf *Surf = &Model->Surfs[i];
+		VectorRef [Surf->vNormal   ] = 0;
+		VectorRef [Surf->vTextureU ] = 0;
+		VectorRef [Surf->vTextureV ] = 0;
+		PointRef  [Surf->pBase     ] = 0;
+	}
+
+	// Check Bsp nodes.
+	for( int32 i=0; i<NumNodes; i++ )
+	{
+		// Tag all points used by nodes.
+		FBspNode*	Node		= &Model->Nodes[i];
+		FVert*		VertPool	= &Model->Verts[Node->iVertPool];
+		for( int B=0; B<Node->NumVertices;  B++ )
+		{
+			PointRef[VertPool->pVertex] = 0;
+			VertPool++;
+		}
+	}
+
+	// Remap points.
+	{
+		int32 n=0;
+		for( int32 i=0; i<NumPoints; i++ ) if( PointRef[i]!=INDEX_NONE )
+		{
+			Model->Points[n] = Model->Points[i];
+			PointRef[i] = n++;
+		}
+		//UE_LOG(LogBSPOps, Log,  TEXT("Points: %i -> %i"), NumPoints, n );
+		Model->Points.RemoveAt( n, NumPoints-n );
+		NumPoints = n;
+	}
+
+	// Remap vectors.
+	{
+		int32 n=0;
+		for (int32 i=0; i<NumVectors; i++) if (VectorRef[i]!=INDEX_NONE)
+		{
+			Model->Vectors[n] = Model->Vectors[i];
+			VectorRef[i] = n++;
+		}
+		//UE_LOG(LogBSPOps, Log,  TEXT("Vectors: %i -> %i"), NumVectors, n );
+		Model->Vectors.RemoveAt( n, NumVectors-n );
+		NumVectors = n;
+	}
+
+	// Update Bsp surfs.
+	for( int32 i=0; i<NumSurfs; i++ )
+	{
+		FBspSurf *Surf	= &Model->Surfs[i];
+		Surf->vNormal   = VectorRef [Surf->vNormal  ];
+		Surf->vTextureU = VectorRef [Surf->vTextureU];
+		Surf->vTextureV = VectorRef [Surf->vTextureV];
+		Surf->pBase     = PointRef  [Surf->pBase    ];
+	}
+
+	// Update Bsp nodes.
+	for( int32 i=0; i<NumNodes; i++ )
+	{
+		FBspNode*	Node		= &Model->Nodes[i];
+		FVert*		VertPool	= &Model->Verts[Node->iVertPool];
+		for( int B=0; B<Node->NumVertices;  B++ )
+		{			
+			VertPool->pVertex = PointRef [VertPool->pVertex];
+			VertPool++;
+		}
+	}
+
+	// Shrink the objects.
+	Model->ShrinkModel();
+
+	Mark.Pop();
+}
+
+// Build bounding volumes for all Bsp nodes.  The bounding volume of the node
+// completely encloses the "outside" space occupied by the nodes.  Note that 
+// this is not the same as representing the bounding volume of all of the 
+// polygons within the node.
+//
+// We start with a practically-infinite cube and filter it down the Bsp,
+// whittling it away until all of its convex volume fragments land in leaves.
+void FBSPOps::bspBuildBounds( UModel* Model )
+{
+	if( Model->Nodes.Num()==0 )
+		return;
+
+	FPoly Polys[6], *PolyList[6];
+	for( int32 i=0; i<6; i++ )
+	{
+		PolyList[i] = &Polys[i];
+		PolyList[i]->Init();
+		PolyList[i]->iBrushPoly = INDEX_NONE;
+	}
+
+	new(Polys[0].Vertices)FVector3f(-HALF_WORLD_MAX,-HALF_WORLD_MAX,HALF_WORLD_MAX);
+	new(Polys[0].Vertices)FVector3f( HALF_WORLD_MAX,-HALF_WORLD_MAX,HALF_WORLD_MAX);
+	new(Polys[0].Vertices)FVector3f( HALF_WORLD_MAX, HALF_WORLD_MAX,HALF_WORLD_MAX);
+	new(Polys[0].Vertices)FVector3f(-HALF_WORLD_MAX, HALF_WORLD_MAX,HALF_WORLD_MAX);
+	Polys[0].Normal   =FVector3f( 0.000000,  0.000000,  1.000000 );
+	Polys[0].Base     =Polys[0].Vertices[0];
+
+	new(Polys[1].Vertices)FVector3f(-HALF_WORLD_MAX, HALF_WORLD_MAX,-HALF_WORLD_MAX);
+	new(Polys[1].Vertices)FVector3f( HALF_WORLD_MAX, HALF_WORLD_MAX,-HALF_WORLD_MAX);
+	new(Polys[1].Vertices)FVector3f( HALF_WORLD_MAX,-HALF_WORLD_MAX,-HALF_WORLD_MAX);
+	new(Polys[1].Vertices)FVector3f(-HALF_WORLD_MAX,-HALF_WORLD_MAX,-HALF_WORLD_MAX);
+	Polys[1].Normal   =FVector3f( 0.000000,  0.000000, -1.000000 );
+	Polys[1].Base     =Polys[1].Vertices[0];
+
+	new(Polys[2].Vertices)FVector3f(-HALF_WORLD_MAX,HALF_WORLD_MAX,-HALF_WORLD_MAX);
+	new(Polys[2].Vertices)FVector3f(-HALF_WORLD_MAX,HALF_WORLD_MAX, HALF_WORLD_MAX);
+	new(Polys[2].Vertices)FVector3f( HALF_WORLD_MAX,HALF_WORLD_MAX, HALF_WORLD_MAX);
+	new(Polys[2].Vertices)FVector3f( HALF_WORLD_MAX,HALF_WORLD_MAX,-HALF_WORLD_MAX);
+	Polys[2].Normal   =FVector3f( 0.000000,  1.000000,  0.000000 );
+	Polys[2].Base     =Polys[2].Vertices[0];
+
+	new(Polys[3].Vertices)FVector3f( HALF_WORLD_MAX,-HALF_WORLD_MAX,-HALF_WORLD_MAX);
+	new(Polys[3].Vertices)FVector3f( HALF_WORLD_MAX,-HALF_WORLD_MAX, HALF_WORLD_MAX);
+	new(Polys[3].Vertices)FVector3f(-HALF_WORLD_MAX,-HALF_WORLD_MAX, HALF_WORLD_MAX);
+	new(Polys[3].Vertices)FVector3f(-HALF_WORLD_MAX,-HALF_WORLD_MAX,-HALF_WORLD_MAX);
+	Polys[3].Normal   =FVector3f( 0.000000, -1.000000,  0.000000 );
+	Polys[3].Base     =Polys[3].Vertices[0];
+
+	new(Polys[4].Vertices)FVector3f(HALF_WORLD_MAX, HALF_WORLD_MAX,-HALF_WORLD_MAX);
+	new(Polys[4].Vertices)FVector3f(HALF_WORLD_MAX, HALF_WORLD_MAX, HALF_WORLD_MAX);
+	new(Polys[4].Vertices)FVector3f(HALF_WORLD_MAX,-HALF_WORLD_MAX, HALF_WORLD_MAX);
+	new(Polys[4].Vertices)FVector3f(HALF_WORLD_MAX,-HALF_WORLD_MAX,-HALF_WORLD_MAX);
+	Polys[4].Normal   =FVector3f( 1.000000,  0.000000,  0.000000 );
+	Polys[4].Base     =Polys[4].Vertices[0];
+
+	new(Polys[5].Vertices)FVector3f(-HALF_WORLD_MAX,-HALF_WORLD_MAX,-HALF_WORLD_MAX);
+	new(Polys[5].Vertices)FVector3f(-HALF_WORLD_MAX,-HALF_WORLD_MAX, HALF_WORLD_MAX);
+	new(Polys[5].Vertices)FVector3f(-HALF_WORLD_MAX, HALF_WORLD_MAX, HALF_WORLD_MAX);
+	new(Polys[5].Vertices)FVector3f(-HALF_WORLD_MAX, HALF_WORLD_MAX,-HALF_WORLD_MAX);
+	Polys[5].Normal   =FVector3f(-1.000000,  0.000000,  0.000000 );
+	Polys[5].Base     =Polys[5].Vertices[0];
+	// Empty hulls.
+	Model->LeafHulls.Empty();
+	for( int32 i=0; i<Model->Nodes.Num(); i++ )
+		Model->Nodes[i].iCollisionBound  = INDEX_NONE;
+	FilterBound( Model, NULL, 0, PolyList, 6, Model->RootOutside );
+//	UE_LOG(LogBSPOps, Log,  TEXT("bspBuildBounds: Generated %i hulls"), Model->LeafHulls.Num() );
+}
+
+/**
+ * Validate a brush, and set iLinks on all EdPolys to index of the
+ * first identical EdPoly in the list, or its index if it's the first.
+ * Not transactional.
+ */
+void FBSPOps::bspValidateBrush( UModel* Brush, bool ForceValidate, bool DoStatusUpdate )
+{
+	check(Brush != nullptr);
+	Brush->Modify();
+	if( ForceValidate || !Brush->Linked )
+	{
+		Brush->Linked = 1;
+		for( int32 i=0; i<Brush->Polys->Element.Num(); i++ )
+		{
+			Brush->Polys->Element[i].iLink = i;
+		}
+		int32 n=0;
+		for( int32 i=0; i<Brush->Polys->Element.Num(); i++ )
+		{
+			FPoly* EdPoly = &Brush->Polys->Element[i];
+			if( EdPoly->iLink==i )
+			{
+				for( int32 j=i+1; j<Brush->Polys->Element.Num(); j++ )
+				{
+					FPoly* OtherPoly = &Brush->Polys->Element[j];
+					if
+					(	OtherPoly->iLink == j
+					&&	OtherPoly->Material == EdPoly->Material
+					&&	OtherPoly->TextureU == EdPoly->TextureU
+					&&	OtherPoly->TextureV == EdPoly->TextureV
+					&&	OtherPoly->PolyFlags == EdPoly->PolyFlags
+					&&	(OtherPoly->Normal | EdPoly->Normal)>0.9999 )
+					{
+						float Dist = FVector3f::PointPlaneDist( OtherPoly->Vertices[0], EdPoly->Vertices[0], EdPoly->Normal );
+						if( Dist>-0.001 && Dist<0.001 )
+						{
+							OtherPoly->iLink = i;
+							n++;
+						}
+					}
+				}
+			}
+		}
+// 		UE_LOG(LogBSPOps, Log,  TEXT("BspValidateBrush linked %i of %i polys"), n, Brush->Polys->Element.Num() );
+	}
+
+	// Build bounds.
+	Brush->BuildBound();
+}
+
+void FBSPOps::bspUnlinkPolys( UModel* Brush )
+{
+	Brush->Modify();
+	Brush->Linked = 1;
+	for( int32 i=0; i<Brush->Polys->Element.Num(); i++ )
+	{
+		Brush->Polys->Element[i].iLink = i;
+	}
+}
+
+// Add an editor polygon to the Bsp, and also stick a reference to it
+// in the editor polygon's BspNodes list. If the editor polygon has more sides
+// than the Bsp will allow, split it up into several sub-polygons.
+//
+// Returns: Index to newly-created node of Bsp.  If several nodes were created because
+// of split polys, returns the parent (highest one up in the Bsp).
+int32	FBSPOps::bspAddNode( UModel* Model, int32 iParent, ENodePlace NodePlace, uint32 NodeFlags, FPoly* EdPoly )
+{
+	if( NodePlace == NODE_Plane )
+	{
+		// Make sure coplanars are added at the end of the coplanar list so that 
+		// we don't insert NF_IsNew nodes with non NF_IsNew coplanar children.
+		while( Model->Nodes[iParent].iPlane != INDEX_NONE )
+		{
+			iParent = Model->Nodes[iParent].iPlane;
+		}
+	}
+	FBspSurf* Surf = NULL;
+	if( EdPoly->iLinkSurf == Model->Surfs.Num() )
+	{
+		int32 NewIndex = Model->Surfs.AddZeroed();
+		Surf = &Model->Surfs[NewIndex];
+
+		// This node has a new polygon being added by bspBrushCSG; must set its properties here.
+		FVector Base = EdPoly->Base, Normal = EdPoly->Normal, TextureU = EdPoly->TextureU, TextureV = EdPoly->TextureV;
+		Surf->pBase     	= bspAddPoint  (Model,&Base,1);
+		Surf->vNormal   	= bspAddVector (Model,&Normal,1);
+		Surf->vTextureU 	= bspAddVector (Model,&TextureU,0);
+		Surf->vTextureV 	= bspAddVector (Model,&TextureV,0);
+		Surf->Material  	= EdPoly->Material;
+		Surf->Actor			= NULL;
+
+		Surf->PolyFlags 	= EdPoly->PolyFlags & ~PF_NoAddToBSP;
+		Surf->LightMapScale= EdPoly->LightMapScale;
+
+		// Find the LightmassPrimitiveSettings in the UModel...
+		int32 FoundLightmassIndex = INDEX_NONE;
+		if (Model->LightmassSettings.Find(EdPoly->LightmassSettings, FoundLightmassIndex) == false)
+		{
+			FoundLightmassIndex = Model->LightmassSettings.Add(EdPoly->LightmassSettings);
+		}
+		Surf->iLightmassIndex = FoundLightmassIndex;
+
+		Surf->Actor	 		= EdPoly->Actor;
+		Surf->iBrushPoly	= EdPoly->iBrushPoly;
+		
+		if (EdPoly->Actor)
+		{
+			Surf->bHiddenEdTemporary = EdPoly->Actor->IsTemporarilyHiddenInEditor();
+			Surf->bHiddenEdLevel = EdPoly->Actor->bHiddenEdLevel;
+			Surf->bHiddenEdLayer = EdPoly->Actor->bHiddenEdLayer;
+		}
+
+		Surf->Plane			= FPlane4f(EdPoly->Vertices[0],EdPoly->Normal);
+	}
+	else
+	{
+		check(EdPoly->iLinkSurf!=INDEX_NONE);
+		check(EdPoly->iLinkSurf<Model->Surfs.Num());
+		Surf = &Model->Surfs[EdPoly->iLinkSurf];
+	}
+
+	// Set NodeFlags.
+	if( Surf->PolyFlags & PF_NotSolid              ) NodeFlags |= NF_NotCsg;
+	if( Surf->PolyFlags & (PF_Invisible|PF_Portal) ) NodeFlags |= NF_NotVisBlocking;
+
+	if( EdPoly->Vertices.Num() > FBspNode::MAX_NODE_VERTICES )
+	{
+		// Split up into two coplanar sub-polygons (one with MAX_NODE_VERTICES vertices and
+		// one with all the remaining vertices) and recursively add them.
+
+		// EdPoly1 is just the first MAX_NODE_VERTICES from EdPoly.
+		FMemMark Mark(FMemStack::Get());
+		FPoly *EdPoly1 = new FPoly;
+		*EdPoly1 = *EdPoly;
+		EdPoly1->Vertices.RemoveAt(FBspNode::MAX_NODE_VERTICES,EdPoly->Vertices.Num() - FBspNode::MAX_NODE_VERTICES);
+
+		// EdPoly2 is the first vertex from EdPoly, and the last EdPoly->Vertices.Num() - MAX_NODE_VERTICES + 1.
+		FPoly *EdPoly2 = new FPoly;
+		*EdPoly2 = *EdPoly;
+		EdPoly2->Vertices.RemoveAt(1,FBspNode::MAX_NODE_VERTICES - 2);
+
+		int32 iNode = bspAddNode( Model, iParent, NodePlace, NodeFlags, EdPoly1 ); // Add this poly first.
+		bspAddNode( Model, iNode,   NODE_Plane, NodeFlags, EdPoly2 ); // Then add other (may be bigger).
+
+		delete EdPoly1;
+		delete EdPoly2;
+
+		Mark.Pop();
+		return iNode; // Return coplanar "parent" node (not coplanar child)
+	}
+	else
+	{
+		// Add node.
+		int32 iNode			 = Model->Nodes.AddZeroed();
+		FBspNode& Node       = Model->Nodes[iNode];
+
+		// Tell transaction tracking system that parent is about to be modified.
+		FBspNode* Parent=NULL;
+		if( NodePlace!=NODE_Root )
+			Parent = &Model->Nodes[iParent];
+
+		// Set node properties.
+		Node.iSurf       	 = EdPoly->iLinkSurf;
+		Node.NodeFlags   	 = NodeFlags;
+		Node.iCollisionBound = INDEX_NONE;
+		Node.Plane           = FPlane4f( EdPoly->Vertices[0], EdPoly->Normal );
+		Node.iVertPool       = Model->Verts.AddUninitialized(EdPoly->Vertices.Num());
+		Node.iFront		     = INDEX_NONE;
+		Node.iBack		     = INDEX_NONE;
+		Node.iPlane		     = INDEX_NONE;
+		if( NodePlace==NODE_Root )
+		{
+			Node.iLeaf[0]	 = INDEX_NONE;
+			Node.iLeaf[1] 	 = INDEX_NONE;
+			Node.iZone[0]	 = 0;
+			Node.iZone[1]	 = 0;
+		}
+		else if( NodePlace==NODE_Front || NodePlace==NODE_Back )
+		{
+			int32 ZoneFront=NodePlace==NODE_Front;
+			Node.iLeaf[0]	 = Parent->iLeaf[ZoneFront];
+			Node.iLeaf[1] 	 = Parent->iLeaf[ZoneFront];
+			Node.iZone[0]	 = Parent->iZone[ZoneFront];
+			Node.iZone[1]	 = Parent->iZone[ZoneFront];
+		}
+		else
+		{
+			int32 IsFlipped    = (Node.Plane|Parent->Plane)<0.0;
+			Node.iLeaf[0]    = Parent->iLeaf[IsFlipped  ];
+			Node.iLeaf[1]    = Parent->iLeaf[1-IsFlipped];
+			Node.iZone[0]    = Parent->iZone[IsFlipped  ];
+			Node.iZone[1]    = Parent->iZone[1-IsFlipped];
+		}
+
+		// Link parent to this node.
+		if (NodePlace == NODE_Front)
+		{
+			Parent->iFront = iNode;
+		}
+		else if (NodePlace == NODE_Back)
+		{
+			Parent->iBack = iNode;
+		}
+		else if (NodePlace == NODE_Plane)
+		{
+			Parent->iPlane = iNode;
+		}
+
+		// Add all points to point table, merging nearly-overlapping polygon points
+		// with other points in the poly to prevent criscrossing vertices from
+		// being generated.
+
+		// Must maintain Node->NumVertices on the fly so that bspAddPoint is always
+		// called with the Bsp in a clean state.
+		Node.NumVertices = 0;
+		FVert* VertPool	 = &Model->Verts[ Node.iVertPool ];
+		for( uint8 i=0; i<EdPoly->Vertices.Num(); i++ )
+		{
+			FVector Vertex = EdPoly->Vertices[i];
+			int32 pVertex = bspAddPoint(Model,&Vertex,0);
+			if( Node.NumVertices==0 || VertPool[Node.NumVertices-1].pVertex!=pVertex )
+			{
+				VertPool[Node.NumVertices].iSide   = INDEX_NONE;
+				VertPool[Node.NumVertices].pVertex = pVertex;
+				Node.NumVertices++;
+			}
+		}
+		if( Node.NumVertices>=2 && VertPool[0].pVertex==VertPool[Node.NumVertices-1].pVertex )
+		{
+			Node.NumVertices--;
+		}
+		if( Node.NumVertices < 3 )
+		{
+			GErrors++;
+// 			UE_LOG(LogBSPOps, Warning, TEXT("bspAddNode: Infinitesimal polygon %i (%i)"), Node.NumVertices, EdPoly->Vertices.Num() );
+			Node.NumVertices = 0;
+		}
+
+		return iNode;
+	}
+}
+
+/**
+ * Rebuild some brush internals
+ */
+void FBSPOps::RebuildBrush(UModel* Brush)
+{
+	Brush->Modify();
+	Brush->EmptyModel(1, 0);
+
+	// Build bounding box.
+	Brush->BuildBound();
+
+	// Build BSP for the brush.
+	bspBuild(Brush, BSP_Good, 15, 70, 1, 0);
+	bspRefresh(Brush, 1);
+	bspBuildBounds(Brush);
+}
+
+/**
+ * Rotates the specified brush's vertices.
+ */
+void FBSPOps::RotateBrushVerts(ABrush* Brush, const FRotator& Rotation, bool bClearComponents)
+{
+	if(Brush->GetBrushComponent()->Brush && Brush->GetBrushComponent()->Brush->Polys)
+	{
+		for( int32 poly = 0 ; poly < Brush->GetBrushComponent()->Brush->Polys->Element.Num() ; poly++ )
+		{
+			FPoly* Poly = &(Brush->GetBrushComponent()->Brush->Polys->Element[poly]);
+
+			// Rotate the vertices.
+			const FRotationMatrix RotMatrix( Rotation );
+			for( int32 vertex = 0 ; vertex < Poly->Vertices.Num() ; vertex++ )
+			{
+				Poly->Vertices[vertex] = Brush->GetPivotOffset() + RotMatrix.TransformVector(Poly->Vertices[vertex] - Brush->GetPivotOffset());
+			}
+			Poly->Base = Brush->GetPivotOffset() + RotMatrix.TransformVector(Poly->Base - Brush->GetPivotOffset());
+
+			// Rotate the texture vectors.
+			Poly->TextureU = RotMatrix.TransformVector( Poly->TextureU );
+			Poly->TextureV = RotMatrix.TransformVector( Poly->TextureV );
+
+			// Recalc the normal for the poly.
+			Poly->Normal = FVector3f::ZeroVector;
+			Poly->Finalize(Brush,0);
+		}
+
+		Brush->GetBrushComponent()->Brush->BuildBound();
+
+		if( !Brush->IsStaticBrush() )
+		{
+			csgPrepMovingBrush( Brush );
+		}
+
+		if ( bClearComponents )
+		{
+			Brush->ReregisterAllComponents();
+		}
+	}
+}
+
+
+void FBSPOps::HandleVolumeShapeChanged(AVolume& Volume)
+{
+	// The default physics volume doesn't have an associated UModel, so we need to handle that case gracefully.
+	if(Volume.Brush)
+	{
+		FBSPOps::csgPrepMovingBrush( &Volume );
+	}
+}
+
+
+void FBspPointsGrid::Clear(int32 InitialSize)
+{
+	GridMap.Empty(InitialSize);
+}
+
+
+// Given a grid index in one axis, a real position on the grid and a threshold radius,
+// return either:
+// - the additional grid index it can overlap in that axis, or
+// - the original grid index if there is no overlap.
+static FORCEINLINE int32 GetAdjacentIndexIfOverlapping(int32 GridIndex, float GridPos, float GridThreshold)
+{
+	if (GridPos - GridIndex < GridThreshold)
+	{
+		return GridIndex - 1;
+	}
+	else if (1.0f - (GridPos - GridIndex) < GridThreshold)
+	{
+		return GridIndex + 1;
+	}
+	else
+	{
+		return GridIndex;
+	}
+}
+
+int32 FBspPointsGrid::FindOrAddPoint(const FVector& Point, int32 Index, float PointThreshold)
+{
+	// Offset applied to the grid coordinates so aligned vertices (the normal case) don't overlap several grid items (taking into account the threshold)
+	const float GridOffset = 0.12345f;
+
+	const float AdjustedPointX = Point.X - GridOffset;
+	const float AdjustedPointY = Point.Y - GridOffset;
+	const float AdjustedPointZ = Point.Z - GridOffset;
+
+	const float GridX = AdjustedPointX * OneOverGranularity;
+	const float GridY = AdjustedPointY * OneOverGranularity;
+	const float GridZ = AdjustedPointZ * OneOverGranularity;
+
+	// Get the grid indices corresponding to the point coordinates
+	const int32 GridIndexX = FMath::FloorToInt(GridX);
+	const int32 GridIndexY = FMath::FloorToInt(GridY);
+	const int32 GridIndexZ = FMath::FloorToInt(GridZ);
+
+	// Find grid item in map
+	FBspPointsGridItem& GridItem = GridMap.FindOrAdd(FBspPointsKey(GridIndexX, GridIndexY, GridIndexZ));
+
+	// Iterate through grid item points and return a point if it's close to the threshold
+	const float PointThresholdSquared = PointThreshold * PointThreshold;
+	for (const FBspIndexedPoint& IndexedPoint : GridItem.IndexedPoints)
+	{
+		if (FVector::DistSquared(IndexedPoint.Point, Point) <= PointThresholdSquared)
+		{
+			return IndexedPoint.Index;
+		}
+	}
+
+	// Otherwise, the point is new: add it to the grid item.
+	GridItem.IndexedPoints.Emplace(Point, Index);
+
+	// The grid has a maximum threshold of a certain radius. If the point is near the edge of a grid cube, it may overlap into other items.
+	// Add it to all grid items it can be seen from.
+	const float GridThreshold = Threshold * OneOverGranularity;
+	const int32 NeighbourX = GetAdjacentIndexIfOverlapping(GridIndexX, GridX, GridThreshold);
+	const int32 NeighbourY = GetAdjacentIndexIfOverlapping(GridIndexY, GridY, GridThreshold);
+	const int32 NeighbourZ = GetAdjacentIndexIfOverlapping(GridIndexZ, GridZ, GridThreshold);
+
+	const bool bOverlapsInX = (NeighbourX != GridIndexX);
+	const bool bOverlapsInY = (NeighbourY != GridIndexY);
+	const bool bOverlapsInZ = (NeighbourZ != GridIndexZ);
+
+	if (bOverlapsInX)
+	{
+		GridMap.FindOrAdd(FBspPointsKey(NeighbourX, GridIndexY, GridIndexZ)).IndexedPoints.Emplace(Point, Index);
+
+		if (bOverlapsInY)
+		{
+			GridMap.FindOrAdd(FBspPointsKey(NeighbourX, NeighbourY, GridIndexZ)).IndexedPoints.Emplace(Point, Index);
+
+			if (bOverlapsInZ)
+			{
+				GridMap.FindOrAdd(FBspPointsKey(NeighbourX, NeighbourY, NeighbourZ)).IndexedPoints.Emplace(Point, Index);
+			}
+		}
+		else if (bOverlapsInZ)
+		{
+			GridMap.FindOrAdd(FBspPointsKey(NeighbourX, GridIndexY, NeighbourZ)).IndexedPoints.Emplace(Point, Index);
+		}
+	}
+	else
+	{
+		if (bOverlapsInY)
+		{
+			GridMap.FindOrAdd(FBspPointsKey(GridIndexX, NeighbourY, GridIndexZ)).IndexedPoints.Emplace(Point, Index);
+
+			if (bOverlapsInZ)
+			{
+				GridMap.FindOrAdd(FBspPointsKey(GridIndexX, NeighbourY, NeighbourZ)).IndexedPoints.Emplace(Point, Index);
+			}
+		}
+		else if (bOverlapsInZ)
+		{
+			GridMap.FindOrAdd(FBspPointsKey(GridIndexX, GridIndexY, NeighbourZ)).IndexedPoints.Emplace(Point, Index);
+		}
+	}
+
+	return Index;
+}
diff --git a/Engine/Source/Developer/BSPUtils/Private/BSPUtils.cpp b/Engine/Source/Developer/BSPUtils/Private/BSPUtils.cpp
new file mode 100644
--- /dev/null
+++ b/Engine/Source/Developer/BSPUtils/Private/BSPUtils.cpp
@@ -0,0 +1,2173 @@
+// Copyright Epic Games, Inc. All Rights Reserved.
+
+#include "BSPUtils.h"
+
+#include "CoreMinimal.h"
+#include "Misc/MemStack.h"
+#include "Misc/FeedbackContext.h"
+#include "Materials/MaterialInterface.h"
+#include "Model.h"
+#include "Materials/Material.h"
+#include "Engine/Polys.h"
+#include "BSPOps.h"
+#include "Engine/Selection.h"
+#include "Modules/ModuleManager.h"
+
+#include "EngineDefines.h"
+#include "Misc/MessageDialog.h"
+#include "Misc/ScopedSlowTask.h"
+#include "GameFramework/Actor.h"
+#include "Engine/Brush.h"
+//#include "Editor/EditorEngine.h"
+//#include "EdMode.h"
+//#include "EditorModeManager.h"
+#include "SurfaceIterators.h"
+//#include "ActorEditorUtils.h"
+#include "Misc/FeedbackContext.h"
+#include "EngineUtils.h"
+
+
+IMPLEMENT_MODULE(FDefaultModuleImpl, BSPUtils)
+
+DEFINE_LOG_CATEGORY_STATIC(LogBspUtils, Log, All);
+
+/*---------------------------------------------------------------------------------------
+	Globals.
+---------------------------------------------------------------------------------------*/
+
+// Magic numbers.
+#define THRESH_OPTGEOM_COPLANAR			(0.25)		/* Threshold for Bsp geometry optimization */
+#define THRESH_OPTGEOM_COSIDAL			(0.25)		/* Threshold for Bsp geometry optimization */
+
+//
+// Status of filtered polygons:
+//
+enum EPolyNodeFilter
+{
+	F_OUTSIDE				= 0, // Leaf is an exterior leaf (visible to viewers).
+	F_INSIDE				= 1, // Leaf is an interior leaf (non-visible, hidden behind backface).
+	F_COPLANAR_OUTSIDE		= 2, // Poly is coplanar and in the exterior (visible to viewers).
+	F_COPLANAR_INSIDE		= 3, // Poly is coplanar and inside (invisible to viewers).
+	F_COSPATIAL_FACING_IN	= 4, // Poly is coplanar, cospatial, and facing in.
+	F_COSPATIAL_FACING_OUT	= 5, // Poly is coplanar, cospatial, and facing out.
+};
+
+//
+// Generic filter function called by BspFilterEdPolys.  A and B are pointers
+// to any integers that your specific routine requires (or NULL if not needed).
+//
+typedef void (*BSP_FILTER_FUNC)
+(
+	UModel			*Model,
+	int32			iNode,
+	FPoly			*EdPoly,
+	EPolyNodeFilter Leaf,
+	ENodePlace		ENodePlace
+);
+
+//
+// Information used by FilterEdPoly.
+//
+class FCoplanarInfo
+{
+public:
+	int32	iOriginalNode;
+	int32	iBackNode;
+	int		BackNodeOutside;
+	int		FrontLeafOutside;
+	int		ProcessingBack;
+};
+
+//
+// Function to filter an EdPoly through the Bsp, calling a callback
+// function for all chunks that fall into leaves.
+//
+void FilterEdPoly
+(
+	BSP_FILTER_FUNC FilterFunc, 
+	UModel			*Model,
+	int32  			iNode, 
+	FPoly			*EdPoly, 
+	FCoplanarInfo	CoplanarInfo, 
+	int				Outside
+);
+
+//
+// Bsp statistics used by link topic function.
+//
+class FBspStats
+{
+public:
+	int   Polys;      // Number of BspSurfs.
+	int   Nodes;      // Total number of Bsp nodes.
+	int   MaxDepth;   // Maximum tree depth.
+	int   AvgDepth;   // Average tree depth.
+	int   Branches;   // Number of nodes with two children.
+	int   Coplanars;  // Number of nodes holding coplanar polygons (in same plane as parent).
+	int   Fronts;     // Number of nodes with only a front child.
+	int   Backs;      // Number of nodes with only a back child.
+	int   Leaves;     // Number of nodes with no children.
+	int   FrontLeaves;// Number of leaf nodes that are in front of their parent.
+	int   BackLeaves; // Number of leaf nodes that are in back of their parent.
+	int   DepthCount; // Depth counter (used for finding average depth).
+} GBspStats;
+
+//
+// Global variables shared between bspBrushCSG and AddWorldToBrushFunc.  These are very
+// tightly tied into the function AddWorldToBrush, not general-purpose.
+//
+int32		GDiscarded;		// Number of polys discarded and not added.
+int32		GNode;          // Node AddBrushToWorld is adding to.
+int32		GLastCoplanar;	// Last coplanar beneath GNode at start of AddWorldToBrush.
+int32		GNumNodes;		// Number of Bsp nodes at start of AddWorldToBrush.
+UModel		*GModel;		// Level map Model we're adding to.
+
+/*----------------------------------------------------------------------------
+   EdPoly building and compacting.
+----------------------------------------------------------------------------*/
+
+//
+// Trys to merge two polygons.  If they can be merged, replaces Poly1 and emptys Poly2
+// and returns 1.  Otherwise, returns 0.
+//
+int TryToMerge( FPoly *Poly1, FPoly *Poly2 )
+{
+	// Find one overlapping point.
+	int32 Start1=0, Start2=0;
+	for( Start1=0; Start1<Poly1->Vertices.Num(); Start1++ )
+		for( Start2=0; Start2<Poly2->Vertices.Num(); Start2++ )
+			if( FVector::PointsAreSame(Poly1->Vertices[Start1], Poly2->Vertices[Start2]) )
+				goto FoundOverlap;
+	return 0;
+	FoundOverlap:
+
+	// Wrap around trying to merge.
+	int32 End1  = Start1;
+	int32 End2  = Start2;
+	int32 Test1 = Start1+1; if (Test1>=Poly1->Vertices.Num()) Test1 = 0;
+	int32 Test2 = Start2-1; if (Test2<0)                   Test2 = Poly2->Vertices.Num()-1;
+	if( FVector::PointsAreSame(Poly1->Vertices[Test1],Poly2->Vertices[Test2]) )
+	{
+		End1   = Test1;
+		Start2 = Test2;
+	}
+	else
+	{
+		Test1 = Start1-1; if (Test1<0)                   Test1=Poly1->Vertices.Num()-1;
+		Test2 = Start2+1; if (Test2>=Poly2->Vertices.Num()) Test2=0;
+		if( FVector::PointsAreSame(Poly1->Vertices[Test1],Poly2->Vertices[Test2]) )
+		{
+			Start1 = Test1;
+			End2   = Test2;
+		}
+		else return 0;
+	}
+
+	// Build a new edpoly containing both polygons merged.
+	FPoly NewPoly = *Poly1;
+	NewPoly.Vertices.Empty();
+	int32 Vertex = End1;
+	for( int32 i=0; i<Poly1->Vertices.Num(); i++ )
+	{
+		new(NewPoly.Vertices) FVector3f(Poly1->Vertices[Vertex]);
+		if( ++Vertex >= Poly1->Vertices.Num() )
+			Vertex=0;
+	}
+	Vertex = End2;
+	for( int32 i=0; i<(Poly2->Vertices.Num()-2); i++ )
+	{
+		if( ++Vertex >= Poly2->Vertices.Num() )
+			Vertex=0;
+		new(NewPoly.Vertices) FVector3f(Poly2->Vertices[Vertex]);
+	}
+
+	// Remove colinear vertices and check convexity.
+	if( NewPoly.RemoveColinears() )
+	{
+		*Poly1 = NewPoly;
+		Poly2->Vertices.Empty();
+		return true;
+	}
+	else return 0;
+}
+
+//
+// Merge all polygons in coplanar list that can be merged convexly.
+//
+void MergeCoplanars( UModel* Model, int32* PolyList, int32 PolyCount )
+{
+	int32 MergeAgain = 1;
+	while( MergeAgain )
+	{
+		MergeAgain = 0;
+		for( int32 i=0; i<PolyCount; i++ )
+		{
+			FPoly& Poly1 = Model->Polys->Element[PolyList[i]];
+			if( Poly1.Vertices.Num() > 0 )
+	        {
+				for( int32 j=i+1; j<PolyCount; j++ )
+				{
+					FPoly& Poly2 = Model->Polys->Element[PolyList[j]];
+					if( Poly2.Vertices.Num() > 0 )
+					{
+                  		if( TryToMerge( &Poly1, &Poly2 ) )
+							MergeAgain=1;
+					}
+				}
+			}
+		}
+	}
+}
+
+//
+// Convert a Bsp node's polygon to an EdPoly, add it to the list, and recurse.
+//
+void MakeEdPolys( UModel* Model, int32 iNode, TArray<FPoly>* DestArray )
+{
+	FBspNode* Node = &Model->Nodes[iNode];
+
+	FPoly Temp;
+	if( FBSPUtils::bspNodeToFPoly(Model,iNode,&Temp) >= 3 )
+	{
+		new(*DestArray)FPoly(Temp);
+	}
+
+	if( Node->iFront!=INDEX_NONE )
+	{
+		MakeEdPolys( Model, Node->iFront, DestArray );
+	}
+	if( Node->iBack !=INDEX_NONE ) 
+	{
+		MakeEdPolys( Model, Node->iBack, DestArray );
+	}
+	if( Node->iPlane!=INDEX_NONE )
+	{
+		MakeEdPolys( Model, Node->iPlane, DestArray );
+	}
+}
+
+void FBSPUtils::bspBuildFPolys( UModel* Model, bool SurfLinks, int32 iNode, TArray<FPoly>* DestArray )
+{
+	if (DestArray == NULL)
+	{
+		DestArray = &Model->Polys->Element;
+	}
+	DestArray->Reset();
+
+	if( Model->Nodes.Num() )
+	{
+		MakeEdPolys( Model, iNode, DestArray );
+	}
+
+	if( !SurfLinks )
+	{
+		const int32 ElementsCount = DestArray->Num();
+		for( int32 i=0; i<ElementsCount; i++ )
+		{
+			(*DestArray)[i].iLink=i;
+		}
+	}
+}
+
+void FBSPUtils::bspMergeCoplanars( UModel* Model, bool RemapLinks, bool MergeDisparateTextures )
+{
+	int32 OriginalNum = Model->Polys->Element.Num();
+
+	// Mark all polys as unprocessed.
+	for( int32 i=0; i<Model->Polys->Element.Num(); i++ )
+		Model->Polys->Element[i].PolyFlags &= ~PF_EdProcessed;
+
+	// Find matching coplanars and merge them.
+	FMemMark Mark(FMemStack::Get());
+	int32* PolyList = new(FMemStack::Get(),Model->Polys->Element.Num())int32;
+	int32 n=0;
+	for( int32 i=0; i<Model->Polys->Element.Num(); i++ )
+	{
+		FPoly* EdPoly = &Model->Polys->Element[i];
+		if( EdPoly->Vertices.Num()>0 && !(EdPoly->PolyFlags & PF_EdProcessed) )
+		{
+			int32 PolyCount         =  0;
+			PolyList[PolyCount++] =  i;
+			EdPoly->PolyFlags    |= PF_EdProcessed;
+			for( int32 j=i+1; j<Model->Polys->Element.Num(); j++ )
+	        {
+				FPoly* OtherPoly = &Model->Polys->Element[j];
+				if( OtherPoly->iLink == EdPoly->iLink && OtherPoly->Vertices.Num() )
+				{
+					float Dist = (OtherPoly->Vertices[0] - EdPoly->Vertices[0]) | EdPoly->Normal;
+					if
+					(	Dist>-0.001
+					&&	Dist<0.001
+					&&	(OtherPoly->Normal|EdPoly->Normal)>0.9999
+					&&	(MergeDisparateTextures
+						||	(	FVector::PointsAreNear(OtherPoly->TextureU,EdPoly->TextureU,THRESH_VECTORS_ARE_NEAR)
+							&&	FVector::PointsAreNear(OtherPoly->TextureV,EdPoly->TextureV,THRESH_VECTORS_ARE_NEAR) ) ) )
+					{
+						OtherPoly->PolyFlags |= PF_EdProcessed;
+						PolyList[PolyCount++] = j;
+					}
+				}
+			}
+			if( PolyCount > 1 )
+	        {
+				MergeCoplanars( Model, PolyList, PolyCount );
+				n++;
+			}
+		}
+	}
+//	UE_LOG(LogEditorBsp, Log,  TEXT("Found %i coplanar sets in %i"), n, Model->Polys->Element.Num() );
+	Mark.Pop();
+
+	// Get rid of empty EdPolys while remapping iLinks.
+	FMemMark Mark2(FMemStack::Get());
+	int32 j=0;
+	int32* Remap = new(FMemStack::Get(),Model->Polys->Element.Num())int32;
+	for( int32 i=0; i<Model->Polys->Element.Num(); i++ )
+	{
+		if( Model->Polys->Element[i].Vertices.Num() )
+		{
+			Remap[i] = j;
+			Model->Polys->Element[j] = Model->Polys->Element[i];
+			j++;
+		}
+	}
+	Model->Polys->Element.RemoveAt( j, Model->Polys->Element.Num()-j );
+	if( RemapLinks )
+	{
+		for( int32 i=0; i<Model->Polys->Element.Num(); i++ )
+		{
+			if (Model->Polys->Element[i].iLink != INDEX_NONE)
+			{
+				CA_SUPPRESS(6001); // warning C6001: Using uninitialized memory '*Remap'.
+				Model->Polys->Element[i].iLink = Remap[Model->Polys->Element[i].iLink];
+			}
+		}
+	}
+// 	UE_LOG(LogEditorBsp, Log,  TEXT("BspMergeCoplanars reduced %i->%i"), OriginalNum, Model->Polys->Element.Num() );
+	Mark2.Pop();
+}
+
+/*----------------------------------------------------------------------------
+   CSG types & general-purpose callbacks.
+----------------------------------------------------------------------------*/
+
+//
+// Recursive worker function called by BspCleanup.
+//
+void CleanupNodes( UModel *Model, int32 iNode, int32 iParent )
+{
+	FBspNode *Node = &Model->Nodes[iNode];
+
+	// Transactionally empty vertices of tag-for-empty nodes.
+	Node->NodeFlags &= ~(NF_IsNew | NF_IsFront | NF_IsBack);
+
+	// Recursively clean up front, back, and plane nodes.
+	if( Node->iFront != INDEX_NONE ) CleanupNodes( Model, Node->iFront, iNode );
+	if( Node->iBack  != INDEX_NONE ) CleanupNodes( Model, Node->iBack , iNode );
+	if( Node->iPlane != INDEX_NONE ) CleanupNodes( Model, Node->iPlane, iNode );
+
+	// Reload Node since the recusive call aliases it.
+	Node = &Model->Nodes[iNode];
+
+	// If this is an empty node with a coplanar, replace it with the coplanar.
+	if( Node->NumVertices==0 && Node->iPlane!=INDEX_NONE )
+	{
+		FBspNode* PlaneNode = &Model->Nodes[ Node->iPlane ];
+
+		// Stick our front, back, and parent nodes on the coplanar.
+		if( (Node->Plane | PlaneNode->Plane) >= 0.0 )
+		{
+			PlaneNode->iFront  = Node->iFront;
+			PlaneNode->iBack   = Node->iBack;
+		}
+		else
+		{
+			PlaneNode->iFront  = Node->iBack;
+			PlaneNode->iBack   = Node->iFront;
+		}
+
+		if( iParent == INDEX_NONE )
+		{
+			// This node is the root.
+			*Node                  = *PlaneNode;   // Replace root.
+			PlaneNode->NumVertices = 0;            // Mark as unused.
+		}
+		else
+		{
+			// This is a child node.
+			FBspNode *ParentNode = &Model->Nodes[iParent];
+
+			if      ( ParentNode->iFront == iNode ) ParentNode->iFront = Node->iPlane;
+			else if ( ParentNode->iBack  == iNode ) ParentNode->iBack  = Node->iPlane;
+			else if ( ParentNode->iPlane == iNode ) ParentNode->iPlane = Node->iPlane;
+			else UE_LOG(LogBspUtils, Fatal, TEXT("CleanupNodes: Parent and child are unlinked"));
+		}
+	}
+	else if( Node->NumVertices == 0 && ( Node->iFront==INDEX_NONE || Node->iBack==INDEX_NONE ) )
+	{
+		// Delete empty nodes with no fronts or backs.
+		// Replace empty nodes with only fronts.
+		// Replace empty nodes with only backs.
+		int32 iReplacementNode;
+		if     ( Node->iFront != INDEX_NONE ) iReplacementNode = Node->iFront;
+		else if( Node->iBack  != INDEX_NONE ) iReplacementNode = Node->iBack;
+		else                                  iReplacementNode = INDEX_NONE;
+
+		if( iParent == INDEX_NONE )
+		{
+			// Root.
+			if( iReplacementNode == INDEX_NONE )
+			{
+         		Model->Nodes.Empty();
+			}
+			else
+			{
+         		*Node = Model->Nodes[iReplacementNode];
+			}
+		}
+		else
+		{
+			// Regular node.
+			FBspNode *ParentNode = &Model->Nodes[iParent];
+
+			if     ( ParentNode->iFront == iNode ) ParentNode->iFront = iReplacementNode;
+			else if( ParentNode->iBack  == iNode ) ParentNode->iBack  = iReplacementNode;
+			else if( ParentNode->iPlane == iNode ) ParentNode->iPlane = iReplacementNode;
+			else UE_LOG(LogBspUtils, Fatal, TEXT("CleanupNodes: Parent and child are unlinked"));
+		}
+	}
+}
+
+
+void FBSPUtils::bspCleanup( UModel *Model )
+{
+	if( Model->Nodes.Num() > 0 ) 
+		CleanupNodes( Model, 0, INDEX_NONE );
+}
+
+/*----------------------------------------------------------------------------
+   CSG leaf filter callbacks.
+----------------------------------------------------------------------------*/
+
+void AddBrushToWorldFunc( UModel* Model, int32 iNode, FPoly* EdPoly,
+	EPolyNodeFilter Filter, ENodePlace ENodePlace )
+{
+	switch( Filter )
+	{
+		case F_OUTSIDE:
+		case F_COPLANAR_OUTSIDE:
+			FBSPOps::bspAddNode (Model,iNode,ENodePlace,NF_IsNew,EdPoly);
+			break;
+		case F_COSPATIAL_FACING_OUT:
+			if( !(EdPoly->PolyFlags & PF_Semisolid) )
+				FBSPOps::bspAddNode (Model,iNode,ENodePlace,NF_IsNew,EdPoly);
+			break;
+		case F_INSIDE:
+		case F_COPLANAR_INSIDE:
+		case F_COSPATIAL_FACING_IN:
+			break;
+	}
+}
+
+void AddWorldToBrushFunc( UModel* Model, int32 iNode, FPoly* EdPoly,
+	EPolyNodeFilter Filter, ENodePlace ENodePlace )
+{
+	switch( Filter )
+	{
+		case F_OUTSIDE:
+		case F_COPLANAR_OUTSIDE:
+			// Only affect the world poly if it has been cut.
+			if( EdPoly->PolyFlags & PF_EdCut )
+				FBSPOps::bspAddNode( GModel, GLastCoplanar, FBSPOps::NODE_Plane, NF_IsNew, EdPoly );
+			break;
+		case F_INSIDE:
+		case F_COPLANAR_INSIDE:
+		case F_COSPATIAL_FACING_IN:
+		case F_COSPATIAL_FACING_OUT:
+			// Discard original poly.
+			GDiscarded++;
+			if( GModel->Nodes[GNode].NumVertices )
+			{
+				GModel->Nodes[GNode].NumVertices = 0;
+			}
+			break;
+	}
+}
+
+void SubtractBrushFromWorldFunc( UModel* Model, int32 iNode, FPoly* EdPoly,
+	EPolyNodeFilter Filter, ENodePlace ENodePlace )
+{
+	switch (Filter)
+	{
+		case F_OUTSIDE:
+		case F_COSPATIAL_FACING_OUT:
+		case F_COSPATIAL_FACING_IN:
+		case F_COPLANAR_OUTSIDE:
+			break;
+		case F_COPLANAR_INSIDE:
+		case F_INSIDE:
+			EdPoly->Reverse();
+			FBSPOps::bspAddNode (Model,iNode,ENodePlace,NF_IsNew,EdPoly); // Add to Bsp back
+			EdPoly->Reverse();
+			break;
+	}
+}
+
+void SubtractWorldToBrushFunc( UModel* Model, int32 iNode, FPoly* EdPoly,
+	EPolyNodeFilter Filter, ENodePlace ENodePlace )
+{
+	switch( Filter )
+	{
+		case F_OUTSIDE:
+		case F_COPLANAR_OUTSIDE:
+		case F_COSPATIAL_FACING_IN:
+			// Only affect the world poly if it has been cut.
+			if( EdPoly->PolyFlags & PF_EdCut )
+				FBSPOps::bspAddNode( GModel, GLastCoplanar, FBSPOps::NODE_Plane, NF_IsNew, EdPoly );
+			break;
+		case F_INSIDE:
+		case F_COPLANAR_INSIDE:
+		case F_COSPATIAL_FACING_OUT:
+			// Discard original poly.
+			GDiscarded++;
+			if( GModel->Nodes[GNode].NumVertices )
+			{
+				GModel->Nodes[GNode].NumVertices = 0;
+			}
+			break;
+	}
+}
+
+void IntersectBrushWithWorldFunc( UModel* Model, int32 iNode, FPoly *EdPoly,
+	EPolyNodeFilter Filter,ENodePlace ENodePlace )
+{
+	switch( Filter )
+	{
+		case F_OUTSIDE:
+		case F_COPLANAR_OUTSIDE:
+		case F_COSPATIAL_FACING_IN:
+		case F_COSPATIAL_FACING_OUT:
+			// Ignore.
+			break;
+		case F_INSIDE:
+		case F_COPLANAR_INSIDE:
+			if( EdPoly->Fix()>=3 )
+				new(GModel->Polys->Element)FPoly(*EdPoly);
+			break;
+	}
+}
+
+void IntersectWorldWithBrushFunc( UModel *Model, int32 iNode, FPoly *EdPoly,
+	EPolyNodeFilter Filter,ENodePlace ENodePlace )
+{
+	switch( Filter )
+	{
+		case F_OUTSIDE:
+		case F_COPLANAR_OUTSIDE:
+		case F_COSPATIAL_FACING_IN:
+			// Ignore.
+			break;
+		case F_INSIDE:
+		case F_COPLANAR_INSIDE:
+		case F_COSPATIAL_FACING_OUT:
+			if( EdPoly->Fix() >= 3 )
+				new(GModel->Polys->Element)FPoly(*EdPoly);
+			break;
+	}
+}
+
+void DeIntersectBrushWithWorldFunc( UModel* Model, int32 iNode, FPoly* EdPoly,
+	EPolyNodeFilter Filter,ENodePlace ENodePlace )
+{
+	switch( Filter )
+	{
+		case F_INSIDE:
+		case F_COPLANAR_INSIDE:
+		case F_COSPATIAL_FACING_OUT:
+		case F_COSPATIAL_FACING_IN:
+			// Ignore.
+			break;
+		case F_OUTSIDE:
+		case F_COPLANAR_OUTSIDE:
+			if( EdPoly->Fix()>=3 )
+				new(GModel->Polys->Element)FPoly(*EdPoly);
+			break;
+	}
+}
+
+void DeIntersectWorldWithBrushFunc( UModel* Model, int32 iNode, FPoly* EdPoly,
+	EPolyNodeFilter Filter,ENodePlace ENodePlace )
+{
+	switch( Filter )
+	{
+		case F_OUTSIDE:
+		case F_COPLANAR_OUTSIDE:
+		case F_COSPATIAL_FACING_OUT:
+			// Ignore.
+			break;
+		case F_COPLANAR_INSIDE:
+		case F_INSIDE:
+		case F_COSPATIAL_FACING_IN:
+			if( EdPoly->Fix() >= 3 )
+	        {
+				EdPoly->Reverse();
+				new(GModel->Polys->Element)FPoly(*EdPoly);
+				EdPoly->Reverse();
+			}
+			break;
+	}
+}
+
+/*----------------------------------------------------------------------------
+   CSG polygon filtering routine (calls the callbacks).
+----------------------------------------------------------------------------*/
+
+//
+// Handle a piece of a polygon that was filtered to a leaf.
+//
+void FilterLeaf
+(
+	BSP_FILTER_FUNC FilterFunc, 
+	UModel*			Model,
+	int32			    iNode, 
+	FPoly*			EdPoly, 
+	FCoplanarInfo	CoplanarInfo, 
+	int32				LeafOutside, 
+	ENodePlace		ENodePlace
+)
+{
+	EPolyNodeFilter FilterType;
+
+	if( CoplanarInfo.iOriginalNode == INDEX_NONE )
+	{
+		// Processing regular, non-coplanar polygons.
+		FilterType = LeafOutside ? F_OUTSIDE : F_INSIDE;
+		FilterFunc( Model, iNode, EdPoly, FilterType, ENodePlace );
+	}
+	else if( CoplanarInfo.ProcessingBack )
+	{
+		// Finished filtering polygon through tree in back of parent coplanar.
+		DoneFilteringBack:
+		if      ((!LeafOutside) && (!CoplanarInfo.FrontLeafOutside)) FilterType = F_COPLANAR_INSIDE;
+		else if (( LeafOutside) && ( CoplanarInfo.FrontLeafOutside)) FilterType = F_COPLANAR_OUTSIDE;
+		else if ((!LeafOutside) && ( CoplanarInfo.FrontLeafOutside)) FilterType = F_COSPATIAL_FACING_OUT;
+		else if (( LeafOutside) && (!CoplanarInfo.FrontLeafOutside)) FilterType = F_COSPATIAL_FACING_IN;
+		else
+		{
+			UE_LOG(LogBspUtils, Fatal, TEXT("FilterLeaf: Bad Locs"));
+			return;
+		}
+		FilterFunc( Model, CoplanarInfo.iOriginalNode, EdPoly, FilterType, FBSPOps::NODE_Plane );
+	}
+	else
+	{
+		CoplanarInfo.FrontLeafOutside = LeafOutside;
+
+		if( CoplanarInfo.iBackNode == INDEX_NONE )
+		{
+			// Back tree is empty.
+			LeafOutside = CoplanarInfo.BackNodeOutside;
+			goto DoneFilteringBack;
+		}
+		else
+		{
+			// Call FilterEdPoly to filter through the back.  This will result in
+			// another call to FilterLeaf with iNode = leaf this falls into in the
+			// back tree and EdPoly = the final EdPoly to insert.
+			CoplanarInfo.ProcessingBack=1;
+			FilterEdPoly( FilterFunc, Model, CoplanarInfo.iBackNode, EdPoly,CoplanarInfo, CoplanarInfo.BackNodeOutside );
+		}
+	}
+}
+
+//
+// Filter an EdPoly through the Bsp recursively, calling FilterFunc
+// for all chunks that fall into leaves.  FCoplanarInfo is used to
+// handle the tricky case of double-recursion for polys that must be
+// filtered through a node's front, then filtered through the node's back,
+// in order to handle coplanar CSG properly.
+//
+void FilterEdPoly
+(
+	BSP_FILTER_FUNC	FilterFunc, 
+	UModel			*Model,
+	int32			    iNode, 
+	FPoly			*EdPoly, 
+	FCoplanarInfo	CoplanarInfo, 
+	int32				Outside
+)
+{
+	int32            SplitResult,iOurFront,iOurBack;
+	int32			   NewFrontOutside,NewBackOutside;
+
+	FilterLoop:
+
+	// Split em.
+	FPoly TempFrontEdPoly,TempBackEdPoly;
+	SplitResult = EdPoly->SplitWithPlane
+	(
+		Model->Points [Model->Verts[Model->Nodes[iNode].iVertPool].pVertex],
+		Model->Vectors[Model->Surfs[Model->Nodes[iNode].iSurf].vNormal],
+		&TempFrontEdPoly,
+		&TempBackEdPoly,
+		0
+	);
+
+	// Process split results.
+	if( SplitResult == SP_Front )
+	{
+		Front:
+
+		FBspNode *Node = &Model->Nodes[iNode];
+		Outside        = Outside || Node->IsCsg();
+
+		if( Node->iFront == INDEX_NONE )
+		{
+			FilterLeaf(FilterFunc,Model,iNode,EdPoly,CoplanarInfo,Outside,FBSPOps::NODE_Front);
+		}
+		else
+		{
+			iNode = Node->iFront;
+			goto FilterLoop;
+		}
+	}
+	else if( SplitResult == SP_Back )
+	{
+		FBspNode *Node = &Model->Nodes[iNode];
+		Outside        = Outside && !Node->IsCsg();
+
+		if( Node->iBack == INDEX_NONE )
+		{
+			FilterLeaf( FilterFunc, Model, iNode, EdPoly, CoplanarInfo, Outside, FBSPOps::NODE_Back );
+		}
+		else
+		{
+			iNode=Node->iBack;
+			goto FilterLoop;
+		}
+	}
+	else if( SplitResult == SP_Coplanar )
+	{
+		if( CoplanarInfo.iOriginalNode != INDEX_NONE )
+		{
+			// This will happen once in a blue moon when a polygon is barely outside the
+			// coplanar threshold and is split up into a new polygon that is
+			// is barely inside the coplanar threshold.  To handle this, just classify
+			// it as front and it will be handled propery.
+			FBSPOps::GErrors++;
+// 			UE_LOG(LogEditorBsp, Warning, TEXT("FilterEdPoly: Encountered out-of-place coplanar") );
+			goto Front;
+		}
+		CoplanarInfo.iOriginalNode        = iNode;
+		CoplanarInfo.iBackNode            = INDEX_NONE;
+		CoplanarInfo.ProcessingBack       = 0;
+		CoplanarInfo.BackNodeOutside      = Outside;
+		NewFrontOutside                   = Outside;
+
+		// See whether Node's iFront or iBack points to the side of the tree on the front
+		// of this polygon (will be as expected if this polygon is facing the same
+		// way as first coplanar in link, otherwise opposite).
+		if( (FVector(Model->Nodes[iNode].Plane) | EdPoly->Normal) >= 0.0 )
+		{
+			iOurFront = Model->Nodes[iNode].iFront;
+			iOurBack  = Model->Nodes[iNode].iBack;
+			
+			if( Model->Nodes[iNode].IsCsg() )
+	        {
+				CoplanarInfo.BackNodeOutside = 0;
+				NewFrontOutside              = 1;
+			}
+		}
+		else
+		{
+			iOurFront = Model->Nodes[iNode].iBack;
+			iOurBack  = Model->Nodes[iNode].iFront;
+			
+			if( Model->Nodes[iNode].IsCsg() )
+	        {
+				CoplanarInfo.BackNodeOutside = 1; 
+				NewFrontOutside              = 0;
+			}
+		}
+
+		// Process front and back.
+		if ((iOurFront==INDEX_NONE)&&(iOurBack==INDEX_NONE))
+		{
+			// No front or back.
+			CoplanarInfo.ProcessingBack		= 1;
+			CoplanarInfo.FrontLeafOutside	= NewFrontOutside;
+			FilterLeaf
+			(
+				FilterFunc,
+				Model,
+				iNode,
+				EdPoly,
+				CoplanarInfo,
+				CoplanarInfo.BackNodeOutside,
+				FBSPOps::NODE_Plane
+			);
+		}
+		else if( iOurFront==INDEX_NONE && iOurBack!=INDEX_NONE )
+		{
+			// Back but no front.
+			CoplanarInfo.ProcessingBack		= 1;
+			CoplanarInfo.iBackNode			= iOurBack;
+			CoplanarInfo.FrontLeafOutside	= NewFrontOutside;
+
+			iNode   = iOurBack;
+			Outside = CoplanarInfo.BackNodeOutside;
+			goto FilterLoop;
+		}
+		else
+		{
+			// Has a front and maybe a back.
+
+			// Set iOurBack up to process back on next call to FilterLeaf, and loop
+			// to process front.  Next call to FilterLeaf will set FrontLeafOutside.
+			CoplanarInfo.ProcessingBack = 0;
+
+			// May be a node or may be INDEX_NONE.
+			CoplanarInfo.iBackNode = iOurBack;
+
+			iNode   = iOurFront;
+			Outside = NewFrontOutside;
+			goto FilterLoop;
+		}
+	}
+	else if( SplitResult == SP_Split )
+	{
+		// Front half of split.
+		if( Model->Nodes[iNode].IsCsg() )
+		{
+			NewFrontOutside = 1; 
+			NewBackOutside  = 0;
+		}
+		else
+		{
+			NewFrontOutside = Outside;
+			NewBackOutside  = Outside;
+		}
+
+		if( Model->Nodes[iNode].iFront==INDEX_NONE )
+		{
+			FilterLeaf
+			(
+				FilterFunc,
+				Model,
+				iNode,
+				&TempFrontEdPoly,
+				CoplanarInfo,
+				NewFrontOutside,
+				FBSPOps::NODE_Front
+			);
+		}
+		else
+		{
+			FilterEdPoly
+			(
+				FilterFunc,
+				Model,
+				Model->Nodes[iNode].iFront,
+				&TempFrontEdPoly,
+				CoplanarInfo,
+				NewFrontOutside
+			);
+		}
+
+		// Back half of split.
+		if( Model->Nodes[iNode].iBack==INDEX_NONE )
+		{
+			FilterLeaf
+			(
+				FilterFunc,
+				Model,
+				iNode,
+				&TempBackEdPoly,
+				CoplanarInfo,
+				NewBackOutside,
+				FBSPOps::NODE_Back
+			);
+		}
+		else
+		{
+			FilterEdPoly
+			(
+				FilterFunc,
+				Model,
+				Model->Nodes[iNode].iBack,
+				&TempBackEdPoly,
+				CoplanarInfo,
+				NewBackOutside
+			);
+		}
+	}
+}
+
+//
+// Regular entry into FilterEdPoly (so higher-level callers don't have to
+// deal with unnecessary info). Filters starting at root.
+//
+void BspFilterFPoly( BSP_FILTER_FUNC FilterFunc, UModel *Model, FPoly *EdPoly )
+{
+	FCoplanarInfo StartingCoplanarInfo;
+	StartingCoplanarInfo.iOriginalNode = INDEX_NONE;
+	if( Model->Nodes.Num() == 0 )
+	{
+		// If Bsp is empty, process at root.
+		FilterFunc( Model, 0, EdPoly, Model->RootOutside ? F_OUTSIDE : F_INSIDE, FBSPOps::NODE_Root );
+	}
+	else
+	{
+		// Filter through Bsp.
+		FilterEdPoly( FilterFunc, Model, 0, EdPoly, StartingCoplanarInfo, Model->RootOutside );
+	}
+}
+
+
+int FBSPUtils::bspNodeToFPoly
+(
+	UModel	*Model,
+	int32	iNode,
+	FPoly	*EdPoly
+)
+{
+	FPoly MasterEdPoly;
+
+	FBspNode &Node     	= Model->Nodes[iNode];
+	FBspSurf &Poly     	= Model->Surfs[Node.iSurf];
+	FVert	 *VertPool	= &Model->Verts[ Node.iVertPool ];
+
+	EdPoly->Base		= Model->Points [Poly.pBase];
+	EdPoly->Normal      = Model->Vectors[Poly.vNormal];
+
+	EdPoly->PolyFlags 	= Poly.PolyFlags & ~(PF_EdCut | PF_EdProcessed | PF_Selected | PF_Memorized);
+	EdPoly->iLinkSurf   = Node.iSurf;
+	EdPoly->Material    = Poly.Material;
+
+	EdPoly->Actor    	= Poly.Actor;
+	EdPoly->iBrushPoly  = Poly.iBrushPoly;
+	
+	if( polyFindMaster(Model,Node.iSurf,MasterEdPoly) )
+		EdPoly->ItemName  = MasterEdPoly.ItemName;
+	else
+		EdPoly->ItemName  = NAME_None;
+
+	EdPoly->TextureU = Model->Vectors[Poly.vTextureU];
+	EdPoly->TextureV = Model->Vectors[Poly.vTextureV];
+
+	EdPoly->LightMapScale = Poly.LightMapScale;
+
+ 	EdPoly->LightmassSettings = Model->LightmassSettings[Poly.iLightmassIndex];
+
+	EdPoly->Vertices.Empty();
+
+	for(int32 VertexIndex = 0;VertexIndex < Node.NumVertices;VertexIndex++)
+	{
+		new(EdPoly->Vertices) FVector3f(Model->Points[VertPool[VertexIndex].pVertex]);
+	}
+
+	if(EdPoly->Vertices.Num() < 3)
+	{
+		EdPoly->Vertices.Empty();
+	}
+	else
+	{
+		// Remove colinear points and identical points (which will appear
+		// if T-joints were eliminated).
+		EdPoly->RemoveColinears();
+	}
+
+	return EdPoly->Vertices.Num();
+}
+
+/*---------------------------------------------------------------------------------------
+   World filtering.
+---------------------------------------------------------------------------------------*/
+
+//
+// Filter all relevant world polys through the brush.
+//
+void FilterWorldThroughBrush 
+(
+	UModel*		Model,
+	UModel*		Brush,
+	EBrushType	BrushType,
+	ECsgOper	CSGOper,
+	int32		    iNode,
+	FSphere*	BrushSphere
+)
+{
+	// Loop through all coplanars.
+	while( iNode != INDEX_NONE )
+	{
+		// Get surface.
+		int32 iSurf = Model->Nodes[iNode].iSurf;
+
+		// Skip new nodes and their children, which are guaranteed new.
+		if( Model->Nodes[iNode].NodeFlags & NF_IsNew )
+			return;
+
+		// Sphere reject.
+		int DoFront = 1, DoBack = 1;
+		if( BrushSphere )
+		{
+			float Dist = Model->Nodes[iNode].Plane.PlaneDot( BrushSphere->Center );
+			DoFront    = (Dist >= -BrushSphere->W);
+			DoBack     = (Dist <= +BrushSphere->W);
+		}
+
+		// Process only polys that aren't empty.
+		FPoly TempEdPoly;
+		if( DoFront && DoBack && (FBSPUtils::bspNodeToFPoly(Model,iNode,&TempEdPoly)>0) )
+		{
+			TempEdPoly.Actor      = Model->Surfs[iSurf].Actor;
+			TempEdPoly.iBrushPoly = Model->Surfs[iSurf].iBrushPoly;
+
+			if( BrushType==Brush_Add || BrushType==Brush_Subtract )
+			{
+				// Add and subtract work the same in this step.
+				GNode       = iNode;
+				GModel  	= Model;
+				GDiscarded  = 0;
+				GNumNodes	= Model->Nodes.Num();
+
+				// Find last coplanar in chain.
+				GLastCoplanar = iNode;
+				while( Model->Nodes[GLastCoplanar].iPlane != INDEX_NONE )
+					GLastCoplanar = Model->Nodes[GLastCoplanar].iPlane;
+
+				// Do the filter operation.
+				BspFilterFPoly
+				(
+					BrushType==Brush_Add ? AddWorldToBrushFunc : SubtractWorldToBrushFunc,
+					Brush,
+					&TempEdPoly
+				);
+				
+				if( GDiscarded == 0 )
+				{
+					// Get rid of all the fragments we added.
+					Model->Nodes[GLastCoplanar].iPlane = INDEX_NONE;
+					const bool bAllowShrinking = false;
+					Model->Nodes.RemoveAt( GNumNodes, Model->Nodes.Num()-GNumNodes, bAllowShrinking );
+				}
+				else
+				{
+					// Tag original world poly for deletion; has been deleted or replaced by partial fragments.
+					if( GModel->Nodes[GNode].NumVertices )
+					{
+						GModel->Nodes[GNode].NumVertices = 0;
+					}
+				}
+			}
+			else if( CSGOper == CSG_Intersect )
+			{
+				BspFilterFPoly( IntersectWorldWithBrushFunc, Brush, &TempEdPoly );
+			}
+			else if( CSGOper == CSG_Deintersect )
+			{
+				BspFilterFPoly( DeIntersectWorldWithBrushFunc, Brush, &TempEdPoly );
+			}
+		}
+
+		// Now recurse to filter all of the world's children nodes.
+		if( DoFront && (Model->Nodes[iNode].iFront != INDEX_NONE)) FilterWorldThroughBrush
+		(
+			Model,
+			Brush,
+			BrushType,
+			CSGOper,
+			Model->Nodes[iNode].iFront,
+			BrushSphere
+		);
+		if( DoBack && (Model->Nodes[iNode].iBack != INDEX_NONE) ) FilterWorldThroughBrush
+		(
+			Model,
+			Brush,
+			BrushType,
+			CSGOper,
+			Model->Nodes[iNode].iBack,
+			BrushSphere
+		);
+		iNode = Model->Nodes[iNode].iPlane;
+	}
+}
+
+
+int FBSPUtils::bspBrushCSG
+(
+	ABrush*		Actor,
+	UModel*		Model,
+	UModel*		TempModel,
+	UMaterialInterface* SelectedMaterialInstance,
+	uint32		PolyFlags, 
+	EBrushType	BrushType,
+	ECsgOper	CSGOper, 
+	bool		bBuildBounds,
+	bool		bMergePolys,
+	bool		bReplaceNULLMaterialRefs,
+	bool		bShowProgressBar/*=true*/
+)
+{
+	uint32 NotPolyFlags = 0;
+	int32 NumPolysFromBrush=0,i,j,ReallyBig;
+	UModel* Brush = Actor->Brush;
+
+	// Note no errors.
+	FBSPOps::GErrors = 0;
+
+	// Make sure we're in an acceptable state.
+	if( !Brush )
+	{
+		return 0;
+	}
+
+	// Non-solid and semisolid stuff can only be added.
+	if( BrushType != Brush_Add )
+	{
+		NotPolyFlags |= (PF_Semisolid | PF_NotSolid);
+	}
+
+	TempModel->EmptyModel(1,1);
+
+	// Update status.
+	ReallyBig = (Brush->Polys->Element.Num() > 200) && bShowProgressBar;
+	if( ReallyBig )
+	{
+		FText Description = NSLOCTEXT("UnrealEd", "PerformingCSGOperation", "Performing CSG operation");
+		
+		if (BrushType != Brush_MAX)
+		{	
+			if (BrushType == Brush_Add)
+			{
+				Description = NSLOCTEXT("UnrealEd", "AddingBrushToWorld", "Adding brush to world");
+			}
+			else if (BrushType == Brush_Subtract)
+			{
+				Description = NSLOCTEXT("UnrealEd", "SubtractingBrushFromWorld", "Subtracting brush from world");
+			}
+		}
+		else if (CSGOper != CSG_None)
+		{
+			if (CSGOper == CSG_Intersect)
+			{
+				Description = NSLOCTEXT("UnrealEd", "IntersectingBrushWithWorld", "Intersecting brush with world");
+			}
+			else if (CSGOper == CSG_Deintersect)
+			{
+				Description = NSLOCTEXT("UnrealEd", "DeintersectingBrushWithWorld", "Deintersecting brush with world");
+			}
+		}
+
+		GWarn->BeginSlowTask( Description, true );
+		// Transform original brush poly into same coordinate system as world
+		// so Bsp filtering operations make sense.
+		GWarn->StatusUpdate(0, 0, NSLOCTEXT("UnrealEd", "Transforming", "Transforming"));
+	}
+
+	const FVector Scale = Actor->GetActorScale();
+	const FRotator Rotation = Actor->GetActorRotation();
+	const FVector Location = Actor->GetActorLocation();
+
+	const bool bIsMirrored = (Scale.X * Scale.Y * Scale.Z < 0.0f);
+
+	// Cache actor transform which is used for the geometry being built
+	Brush->OwnerLocationWhenLastBuilt = Location;
+	Brush->OwnerRotationWhenLastBuilt = Rotation;
+	Brush->OwnerScaleWhenLastBuilt = Scale;
+	Brush->bCachedOwnerTransformValid = true;
+
+	for( i=0; i<Brush->Polys->Element.Num(); i++ )
+	{
+		FPoly& CurrentPoly = Brush->Polys->Element[i];
+
+		// Set texture the first time.
+		if ( bReplaceNULLMaterialRefs )
+		{
+			UMaterialInterface*& PolyMat = CurrentPoly.Material;
+			if ( !PolyMat || PolyMat == UMaterial::GetDefaultMaterial(MD_Surface) )
+			{
+				PolyMat = SelectedMaterialInstance;
+			}
+		}
+
+		// Get the brush poly.
+		FPoly DestEdPoly = CurrentPoly;
+		check(CurrentPoly.iLink<Brush->Polys->Element.Num());
+
+		// Set its backward brush link.
+		DestEdPoly.Actor       = Actor;
+		DestEdPoly.iBrushPoly  = i;
+
+		// Update its flags.
+		DestEdPoly.PolyFlags = (DestEdPoly.PolyFlags | PolyFlags) & ~NotPolyFlags;
+
+		// Set its internal link.
+		if (DestEdPoly.iLink == INDEX_NONE)
+		{
+			DestEdPoly.iLink = i;
+		}
+
+		// Transform it.
+		DestEdPoly.Scale( Scale );
+		DestEdPoly.Rotate( Rotation );
+		DestEdPoly.Transform( Location );
+
+		// Reverse winding and normal if the parent brush is mirrored
+		if (bIsMirrored)
+		{
+			DestEdPoly.Reverse();
+			DestEdPoly.CalcNormal();
+		}
+
+		// Add poly to the temp model.
+		new(TempModel->Polys->Element)FPoly( DestEdPoly );
+	}
+	if( ReallyBig ) GWarn->StatusUpdate( 0, 0, NSLOCTEXT("UnrealEd", "FilteringBrush", "Filtering brush") );
+
+	// Pass the brush polys through the world Bsp.
+	if( CSGOper==CSG_Intersect || CSGOper==CSG_Deintersect )
+	{
+		// Empty the brush.
+		Brush->EmptyModel(1,1);
+
+		// Intersect and deintersect.
+		for( i=0; i<TempModel->Polys->Element.Num(); i++ )
+		{
+         	FPoly EdPoly = TempModel->Polys->Element[i];
+			GModel = Brush;
+			// TODO: iLink / iLinkSurf in EdPoly / TempModel->Polys->Element[i] ?
+			BspFilterFPoly( CSGOper==CSG_Intersect ? IntersectBrushWithWorldFunc : DeIntersectBrushWithWorldFunc, Model,  &EdPoly );
+		}
+		NumPolysFromBrush = Brush->Polys->Element.Num();
+	}
+	else
+	{
+		// Add and subtract.
+		TMap<int32, int32> SurfaceIndexRemap;
+		for( i=0; i<TempModel->Polys->Element.Num(); i++ )
+		{
+         	FPoly EdPoly = TempModel->Polys->Element[i];
+
+         	// Mark the polygon as non-cut so that it won't be harmed unless it must
+         	// be split, and set iLink so that BspAddNode will know to add its information
+         	// if a node is added based on this poly.
+         	EdPoly.PolyFlags &= ~(PF_EdCut);
+			const int32* SurfaceIndexPtr = SurfaceIndexRemap.Find(EdPoly.iLink);
+			if (SurfaceIndexPtr == nullptr)
+			{
+				const int32 NewSurfaceIndex = Model->Surfs.Num();
+				SurfaceIndexRemap.Add(EdPoly.iLink, NewSurfaceIndex);
+				EdPoly.iLinkSurf = TempModel->Polys->Element[i].iLinkSurf = NewSurfaceIndex;
+			}
+			else
+			{
+				EdPoly.iLinkSurf = TempModel->Polys->Element[i].iLinkSurf = *SurfaceIndexPtr;
+			}
+
+			// Filter brush through the world.
+			BspFilterFPoly( BrushType==Brush_Add ? AddBrushToWorldFunc : SubtractBrushFromWorldFunc, Model, &EdPoly );
+		}
+	}
+	if( Model->Nodes.Num() && !(PolyFlags & (PF_NotSolid | PF_Semisolid)) )
+	{
+		// Quickly build a Bsp for the brush, tending to minimize splits rather than balance
+		// the tree.  We only need the cutting planes, though the entire Bsp struct (polys and
+		// all) is built.
+
+		FBspPointsGrid* LevelModelPointsGrid = FBspPointsGrid::GBspPoints;
+		FBspPointsGrid* LevelModelVectorsGrid = FBspPointsGrid::GBspVectors;
+
+		// For the bspBuild call, temporarily create a new pair of BspPointsGrids for the TempModel.
+		TUniquePtr<FBspPointsGrid> BspPoints = MakeUnique<FBspPointsGrid>(50.0f, THRESH_POINTS_ARE_SAME);
+ 		TUniquePtr<FBspPointsGrid> BspVectors = MakeUnique<FBspPointsGrid>(1 / 16.0f, FMath::Max(THRESH_NORMALS_ARE_SAME, THRESH_VECTORS_ARE_NEAR));
+ 		FBspPointsGrid::GBspPoints = BspPoints.Get();
+ 		FBspPointsGrid::GBspVectors = BspVectors.Get();
+
+		if( ReallyBig ) GWarn->StatusUpdate( 0, 0, NSLOCTEXT("UnrealEd", "BuildingBSP", "Building BSP") );
+		
+		FBSPOps::bspBuild( TempModel, FBSPOps::BSP_Lame, 0, 70, 1, 0 );
+
+		// Reinstate the original BspPointsGrids used for building the level Model.
+		FBspPointsGrid::GBspPoints = LevelModelPointsGrid;
+		FBspPointsGrid::GBspVectors = LevelModelVectorsGrid;
+
+		if( ReallyBig ) GWarn->StatusUpdate( 0, 0, NSLOCTEXT("UnrealEd", "FilteringWorld", "Filtering world") );
+		GModel = Brush;
+		TempModel->BuildBound();
+
+		FSphere	BrushSphere = TempModel->Bounds.GetSphere();
+		FilterWorldThroughBrush( Model, TempModel, BrushType, CSGOper, 0, &BrushSphere );
+	}
+	if( CSGOper==CSG_Intersect || CSGOper==CSG_Deintersect )
+	{
+		if( ReallyBig ) GWarn->StatusUpdate( 0, 0, NSLOCTEXT("UnrealEd", "AdjustingBrush", "Adjusting brush") );
+
+		// Link polys obtained from the original brush.
+		for( i=NumPolysFromBrush-1; i>=0; i-- )
+		{
+			FPoly *DestEdPoly = &Brush->Polys->Element[i];
+			for( j=0; j<i; j++ )
+			{
+				if( DestEdPoly->iLink == Brush->Polys->Element[j].iLink )
+				{
+					DestEdPoly->iLink = j;
+					break;
+				}
+			}
+			if( j >= i ) DestEdPoly->iLink = i;
+		}
+
+		// Link polys obtained from the world.
+		for( i=Brush->Polys->Element.Num()-1; i>=NumPolysFromBrush; i-- )
+		{
+			FPoly *DestEdPoly = &Brush->Polys->Element[i];
+			for( j=NumPolysFromBrush; j<i; j++ )
+			{
+				if( DestEdPoly->iLink == Brush->Polys->Element[j].iLink )
+				{
+					DestEdPoly->iLink = j;
+					break;
+				}
+			}
+			if( j >= i ) DestEdPoly->iLink = i;
+		}
+		Brush->Linked = 1;
+
+		// Detransform the obtained brush back into its original coordinate system.
+		for( i=0; i<Brush->Polys->Element.Num(); i++ )
+		{
+			FPoly *DestEdPoly = &Brush->Polys->Element[i];
+			DestEdPoly->Transform(-Location);
+			DestEdPoly->Rotate(Rotation.GetInverse());
+			DestEdPoly->Scale(FVector(1.0f) / Scale);
+			DestEdPoly->Fix();
+			DestEdPoly->Actor		= NULL;
+			DestEdPoly->iBrushPoly	= i;
+		}
+	}
+
+	if( BrushType==Brush_Add || BrushType==Brush_Subtract )
+	{
+		// Clean up nodes, reset node flags.
+		bspCleanup( Model );
+
+		// Rebuild bounding volumes.
+		if( bBuildBounds )
+		{
+			FBSPOps::bspBuildBounds( Model );
+		}
+	}
+
+	Brush->NumUniqueVertices = TempModel->Points.Num();
+	// Release TempModel.
+	TempModel->EmptyModel(1,1);
+	
+	// Merge coplanars if needed.
+	if( CSGOper==CSG_Intersect || CSGOper==CSG_Deintersect )
+	{
+		if( ReallyBig )
+		{
+			GWarn->StatusUpdate( 0, 0, NSLOCTEXT("UnrealEd", "Merging", "Merging") );
+		}
+		if( bMergePolys )
+		{
+			bspMergeCoplanars( Brush, 1, 0 );
+		}
+	}
+	if( ReallyBig )
+	{
+		GWarn->EndSlowTask();
+	}
+
+	return 1 + FBSPOps::GErrors;
+}
+
+/*---------------------------------------------------------------------------------------
+   Functions for maintaining linked geometry lists.
+---------------------------------------------------------------------------------------*/
+
+//
+// A node and vertex number corresponding to a point, used in generating Bsp side links.
+//
+class FPointVert 
+{
+public:
+	int32		    iNode;
+	int32		    nVertex;
+	FPointVert* Next;
+};
+
+//
+// A list of point/vertex links, used in generating Bsp side links.
+//
+class FPointVertList
+{
+public:
+
+	// Variables.
+	UModel		*Model;
+	FPointVert	**Index;
+	FMemMark*	Mark;
+
+	/** Initialization constructor. */
+	FPointVertList()
+	:	Model(NULL)
+	,	Index(NULL)
+	,	Mark(NULL)
+	{}
+
+	/** Destructor. */
+	~FPointVertList()
+	{
+		check(!Mark);
+	}
+
+	// Allocate.
+	void Alloc( UModel *ThisModel )
+	{
+		check(!Mark);
+		Mark = new FMemMark(FMemStack::Get());
+
+		// Allocate working tables.
+		Model = ThisModel;
+		Index = new(FMemStack::Get(),MEM_Zeroed,Model->Points.Num())FPointVert*;
+
+	}
+
+	// Free.
+	void Free()
+	{
+		delete Mark;
+		Mark = NULL;
+	}
+
+	// Add all of a node's vertices to a node-vertex list.
+	void AddNode( int32 iNode )
+	{
+		FBspNode&	Node 	 =  Model->Nodes[iNode];
+		FVert*		VertPool =  &Model->Verts[Node.iVertPool];
+
+		for( int32 i=0; i < Node.NumVertices; i++ )
+		{
+			int32 pVertex = VertPool[i].pVertex;
+
+			// Add new point/vertex pair to array, and insert new array entry
+			// between index and first entry.
+			FPointVert *PointVert	= new(FMemStack::Get())FPointVert;
+			
+			PointVert->iNode 		= iNode;
+			PointVert->nVertex		= i;
+			PointVert->Next			= Index[pVertex];
+
+			Index[pVertex]			= PointVert;
+		}
+	}
+
+	// Add all nodes' vertices in the model to a node-vertex list.
+	void AddAllNodes()
+	{
+		for( int32 iNode=0; iNode < Model->Nodes.Num(); iNode++ )
+			AddNode( iNode );
+
+	}
+
+	// Remove all of a node's vertices from a node-vertex list.
+	void RemoveNode( int32 iNode )
+	{
+		FBspNode&	Node 		= Model->Nodes[iNode];
+		FVert*		VertPool	= &Model->Verts[Node.iVertPool];
+
+		// Loop through all of the node's vertices and search through the
+		// corresponding point's node-vert list, and delink this node.
+		int Count=0;
+		for( int32 i = 0; i < Node.NumVertices; i++ )
+		{
+			int32 pVertex = VertPool[i].pVertex;
+
+			for( FPointVert **PrevLink = &Index[pVertex]; *PrevLink; PrevLink=&(*PrevLink)->Next )
+			{
+				if( (*PrevLink)->iNode == iNode )
+				{
+					// Delink this entry from the list.
+					*PrevLink = (*PrevLink)->Next;
+					Count++;
+
+					if( *PrevLink == NULL )
+						break;
+				}
+			}
+
+			// Node's vertex wasn't found, there's a bug.
+			check(Count>=1);
+		}
+	}
+};
+
+/*---------------------------------------------------------------------------------------
+   Geometry optimization.
+---------------------------------------------------------------------------------------*/
+
+//
+// Add a point to a Bsp node before a specified vertex (between it and the previous one).
+// VertexNumber can be from 0 (before first) to Node->NumVertices (after last).
+//
+// Splits node into two coplanar polys if necessary. If the polygon is split, the
+// vertices will be distributed among this node and it's newly-linked iPlane node
+// in an arbitrary way, that preserves the clockwise orientation of the vertices.
+//
+// Maintains node-vertex list, if not NULL.
+//
+void AddPointToNode
+(
+	UModel*         Model,
+	FPointVertList* PointVerts,
+	int32			    iNode, 
+	int32			    VertexNumber, 
+	int32			    pVertex 
+)
+{
+	FBspNode&	Node = Model->Nodes[iNode];
+
+	if( Node.NumVertices >= FBspNode::MAX_NODE_VERTICES - 1 )
+	{
+		// Just refuse to add point: This is a non-fatal problem.
+//		UE_LOG(LogEditorBsp, Warning, TEXT("Node side limit reached") );
+		return;
+	}
+
+	// Remove node from vertex list, since vertex numbers will be reordered.
+	if( PointVerts )
+		PointVerts->RemoveNode( iNode );
+
+	int32 iOldVert = Node.iVertPool;
+
+	Node.iVertPool = Model->Verts.AddUninitialized( Node.NumVertices+1 );
+
+	// Make sure this node doesn't already contain the vertex.
+	for( int32 i=0; i<Node.NumVertices; i++ )
+		check( Model->Verts[iOldVert + i].pVertex != pVertex );
+
+	// Copy the old vertex pool to the new one.
+	for( int32 i=0; i<VertexNumber; i++ )
+		Model->Verts[Node.iVertPool + i] = Model->Verts[iOldVert + i];
+
+	for( int32 i=VertexNumber; i<Node.NumVertices; i++ )
+		Model->Verts[Node.iVertPool + i + 1] = Model->Verts[iOldVert + i];
+
+	// Add the new point to the new vertex pool.
+	Model->Verts[Node.iVertPool + VertexNumber].pVertex = pVertex;
+	Model->Verts[Node.iVertPool + VertexNumber].iSide   = INDEX_NONE;
+
+	// Increment number of node vertices.
+	Node.NumVertices++;
+
+	// Update the point-vertex list.
+	if( PointVerts )
+		PointVerts->AddNode( iNode );
+
+}
+
+//
+// Add a point to all sides of polygons in which the side intersects with
+// this point but doesn't contain it, and has the correct (clockwise) orientation
+// as this side.  pVertex is the index of the point to handle, and
+// ReferenceVertex defines the direction of this side.
+//
+int DistributePoint
+(
+	UModel*         Model,
+	FPointVertList* PointVerts,	
+	int32			    iNode,
+	int32			    pVertex
+)
+{
+	int32 Count = 0;
+
+	// Handle front, back, and plane.
+	float Dist = Model->Nodes[iNode].Plane.PlaneDot(Model->Points[pVertex]);
+	if( Dist < THRESH_OPTGEOM_COPLANAR )
+	{
+		// Back.
+		if( Model->Nodes[iNode].iBack != INDEX_NONE )
+			Count += DistributePoint( Model, PointVerts, Model->Nodes[iNode].iBack, pVertex );
+	}
+	if( Dist > -THRESH_OPTGEOM_COPLANAR )
+	{
+		// Front.
+		if( Model->Nodes[iNode].iFront != INDEX_NONE )
+			Count += DistributePoint( Model, PointVerts, Model->Nodes[iNode].iFront, pVertex );
+	}
+	if( Dist > -THRESH_OPTGEOM_COPLANAR && Dist < THRESH_OPTGEOM_COPLANAR )
+	{
+		// This point is coplanar with this node, so check point for intersection with
+		// this node's sides, then loop with its coplanars.
+		for( ; iNode!=INDEX_NONE; iNode=Model->Nodes[iNode].iPlane )
+		{
+			FVert* VertPool = &Model->Verts[Model->Nodes[iNode].iVertPool];
+
+			// Skip this node if it already contains the point in question.
+			int32 i;
+			for( i=0; i<Model->Nodes[iNode].NumVertices; i++ )
+				if( VertPool[i].pVertex == pVertex )
+					break;
+			if( i != Model->Nodes[iNode].NumVertices )
+				continue;
+
+			// Loop through all sides and see if (A) side is colinear with point, and
+			// (B) point falls within inside of this side.
+			int32 FoundSide       = -1;
+			int32 SkippedColinear = 0;
+			int32 SkippedInside   = 0;
+
+			for( i=0; i<Model->Nodes[iNode].NumVertices; i++ )
+			{
+				int32 j = (i>0) ? (i-1) : (Model->Nodes[iNode].NumVertices-1);
+
+				// Create cutting plane perpendicular to both this side and the polygon's normal.
+				FVector Side            = Model->Points[VertPool[i].pVertex] - Model->Points[VertPool[j].pVertex];
+				FVector SidePlaneNormal = Side ^ Model->Nodes[iNode].Plane;
+				float   SizeSquared     = SidePlaneNormal.SizeSquared();
+
+				if( SizeSquared>FMath::Square(0.001) )
+				{
+					// Points aren't coincedent.
+					Dist = ((Model->Points[pVertex]-Model->Points[VertPool[i].pVertex]) | SidePlaneNormal)/FMath::Sqrt(SizeSquared);
+					if( Dist >= THRESH_OPTGEOM_COSIDAL )
+					{
+						// Point is outside polygon, can't possibly fall on a side.
+						break;
+					}
+					else if( Dist > -THRESH_OPTGEOM_COSIDAL )
+					{
+						// The point we're adding falls on this line.
+						//
+						// Verify that it falls within this side; though it's colinear
+						// it may be out of the bounds of the line's endpoints if this side
+						// is colinear with an adjacent side.
+						//
+						// Do this by checking distance from point to side's midpoint and
+						// comparing with the side's half-length.
+
+						FVector MidPoint    = (Model->Points[VertPool[i].pVertex] + Model->Points[VertPool[j].pVertex])*0.5;
+						FVector MidDistVect = Model->Points[pVertex] - MidPoint;
+						if( MidDistVect.SizeSquared() <= FMath::Square(0.501)*Side.SizeSquared() )
+						{
+							FoundSide = i;
+						}
+						else
+						{
+							SkippedColinear = 1;
+						}
+					}
+					else
+					{
+						// Point is inside polygon, so continue checking.
+						SkippedInside = 1;
+					}
+				}
+				else
+				{
+					FBSPOps::GErrors++;
+					//UE_LOG(LogEditorBsp, Warning, "Found tiny side" );
+					//Poly.PolyFlags |= PF_Selected;
+				}
+			}
+			if( i==Model->Nodes[iNode].NumVertices && FoundSide>=0 )
+			{
+				// AddPointToNode will reorder the vertices in this node.  This is okay
+				// because it's called outside of the vertex loop.
+				AddPointToNode( Model, PointVerts, iNode, FoundSide, pVertex );
+				Count++;
+			}
+			else if( SkippedColinear )
+			{
+				// This happens occasionally because of the fuzzy Dist comparison.  It is
+				// not a sign of a problem when the vertex being distributed is colinear
+				// with one of this polygon's sides, but slightly outside of this polygon.
+
+				FBSPOps::GErrors++;
+				//UE_LOG(LogEditorBsp, Warning, "Skipped colinear" );
+				//Poly.PolyFlags |= PF_Selected;
+			}
+			else if( SkippedInside )
+			{
+				// Point is on interior of polygon.
+				FBSPOps::GErrors++;
+				//UE_LOG(LogEditorBsp, Warning, "Skipped interior" );
+				//Poly.PolyFlags |= PF_Selected;
+			}
+		}
+	}
+	return Count;
+}
+
+//
+// Merge points that are near.
+//
+void MergeNearPoints( UModel *Model, float Dist )
+{
+	FMemMark Mark(FMemStack::Get());
+	int32* PointRemap = new(FMemStack::Get(),Model->Points.Num())int32;
+	int32 Merged=0,Collapsed=0;
+
+	// Find nearer point for all points.
+	for( int32 i=0; i<Model->Points.Num(); i++ )
+	{
+		PointRemap[i] = i;
+		FVector3f &Point = Model->Points[i];
+		for( int32 j=0; j<i; j++ )
+		{
+			FVector3f &TestPoint = Model->Points[j];
+			if( (TestPoint - Point).SizeSquared() < Dist*Dist )
+			{
+				PointRemap[i] = j;
+				Merged++;
+				break;
+			}
+		}
+	}
+
+	// Remap VertPool.
+	for( int32 i=0; i<Model->Verts.Num(); i++ )
+	{
+		if( Model->Verts[i].pVertex>=0 && Model->Verts[i].pVertex<Model->Points.Num() )
+		{
+			CA_SUPPRESS(6001); // warning C6001: Using uninitialized memory '*PointRemap'.
+			Model->Verts[i].pVertex = PointRemap[Model->Verts[i].pVertex];
+		}
+	}
+
+	// Remap Surfs.
+	for( int32 i=0; i<Model->Surfs.Num(); i++ )
+	{
+		if( Model->Surfs[i].pBase>=0 && Model->Surfs[i].pBase<Model->Points.Num() )
+		{
+			CA_SUPPRESS(6001); // warning C6001: Using uninitialized memory '*PointRemap'.
+			Model->Surfs[i].pBase = PointRemap[Model->Surfs[i].pBase];
+		}
+	}
+
+	// Remove duplicate points from nodes.
+	for( int32 i=0; i<Model->Nodes.Num(); i++ )
+	{
+		FBspNode &Node = Model->Nodes[i];
+		FVert *Pool = &Model->Verts[Node.iVertPool];
+		
+		int k=0;
+		for( int j=0; j<Node.NumVertices; j++ )
+		{
+			FVert *A = &Pool[j];
+			FVert *B = &Pool[j ? j-1 : Node.NumVertices-1];
+			
+			if( A->pVertex != B->pVertex )
+				Pool[k++] = Pool[j];
+		}
+		Node.NumVertices = k>=3 ? k : 0;
+		if( k < 3 )
+			Collapsed++;
+	}
+
+	// Success.
+// 	UE_LOG(LogEditorBsp, Log,  TEXT("MergeNearPoints merged %i/%i, collapsed %i"), Merged, Model->Points.Num(), Collapsed );
+	Mark.Pop();
+}
+
+
+void FBSPUtils::bspOptGeom( UModel *Model )
+{
+	FPointVertList PointVerts;
+
+	//UE_LOG(LogEditorBsp, Log,  TEXT("BspOptGeom begin") );
+
+	MergeNearPoints			(Model,0.25);
+	FBSPOps::bspRefresh		(Model,0);
+	PointVerts.Alloc		(Model);
+	PointVerts.AddAllNodes	();
+
+	// First four entries are reserved for view-clipped sides.
+	Model->NumSharedSides = 4;
+	
+	// Mark all sides as unlinked.
+	int32 i;
+	for( i=0; i<Model->Verts.Num(); i++ )
+		Model->Verts[i].iSide = INDEX_NONE;
+
+	int32 TeesFound = 0, Distributed = 0;
+
+	// Eliminate T-joints on each node by finding all vertices that aren't attached to
+	// two shared sides, then filtering them down through the BSP and adding them to
+	// the sides they belong on.
+	for( int32 iNode=0; iNode < Model->Nodes.Num(); iNode++ )
+	{
+		FBspNode &Node = Model->Nodes[iNode];
+		
+		// Loop through all sides (side := line from PrevVert to ThisVert)		
+		for( uint8 ThisVert=0; ThisVert < Node.NumVertices; ThisVert++ )
+		{
+			uint8 PrevVert = (ThisVert>0) ? (ThisVert - 1) : (Node.NumVertices-1);
+
+			// Count number of nodes sharing this side, i.e. number of nodes for
+			// which two adjacent vertices are identical to this side's two vertices.
+			for( FPointVert* PV1 = PointVerts.Index[Model->Verts[Node.iVertPool+ThisVert].pVertex]; PV1; PV1=PV1->Next )
+				for( FPointVert* PV2 = PointVerts.Index[Model->Verts[Node.iVertPool+PrevVert].pVertex]; PV2; PV2=PV2->Next )
+					if( PV1->iNode==PV2->iNode && PV1->iNode!=iNode )
+						goto SkipIt;
+
+			// Didn't find another node that shares our two vertices; must add each
+			// vertex to all polygons where the vertex lies on the polygon's side.
+			// DistributePoint will not affect the current node but may change others
+			// and may increase the number of nodes in the Bsp.
+			TeesFound++;
+			Distributed = 0;
+			Distributed += DistributePoint( Model, &PointVerts, 0, Model->Verts[Node.iVertPool+ThisVert].pVertex );
+			Distributed += DistributePoint( Model, &PointVerts, 0, Model->Verts[Node.iVertPool+PrevVert].pVertex );
+			SkipIt:;
+		}
+	}
+	// Build side links
+	// Definition of side: Side (i) links node vertex (i) to vertex ((i+1)%n)
+	//UE_LOG(LogEditorBsp, Log,  TEXT("BspOptGeom building sidelinks") );
+
+	PointVerts.Free			();
+	PointVerts.Alloc		(Model);
+	PointVerts.AddAllNodes	();
+
+	for( int32 iNode=0; iNode < Model->Nodes.Num(); iNode++ )
+	{
+		FBspNode &Node = Model->Nodes[iNode];		
+		for( uint8 ThisVert=0; ThisVert < Node.NumVertices; ThisVert++ )
+		{
+			if( Model->Verts[Node.iVertPool+ThisVert].iSide == INDEX_NONE )
+			{
+				// See if this node links to another one.
+				uint8 PrevVert = (ThisVert>0) ? (ThisVert - 1) : (Node.NumVertices-1);
+				for( FPointVert* PV1=PointVerts.Index[Model->Verts[Node.iVertPool+ThisVert].pVertex]; PV1; PV1=PV1->Next )
+				{
+					for( FPointVert* PV2=PointVerts.Index[Model->Verts[Node.iVertPool+PrevVert].pVertex]; PV2; PV2=PV2->Next )
+					{
+						if( PV1->iNode==PV2->iNode && PV1->iNode!=iNode )
+						{
+							// Make sure that the other node's two vertices are adjacent and
+							// ordered opposite this node's vertices.
+							int32		    iOtherNode	= PV2->iNode;
+							FBspNode	&OtherNode	= Model->Nodes[iOtherNode];
+
+							int32 Delta = 
+								(OtherNode.NumVertices + PV2->nVertex - PV1->nVertex) % 
+								OtherNode.NumVertices;
+
+							if( Delta==1 )
+							{
+								// Side is properly linked!
+								int32 OtherVert = PV2->nVertex;
+								int32 iSide;
+								if( Model->Verts[OtherNode.iVertPool+OtherVert].iSide==INDEX_NONE )
+									iSide = Model->NumSharedSides++;
+								else
+									iSide = Model->Verts[OtherNode.iVertPool+OtherVert].iSide;
+									
+								// Link both sides to the shared side.
+								Model->Verts[ Node.iVertPool + ThisVert ].iSide
+								=	Model->Verts[ OtherNode.iVertPool+OtherVert ].iSide
+								=	iSide;
+								goto SkipSide;
+							}
+						}
+					}
+				}
+
+				// This node doesn't have correct side linking
+				//Model->Surfs(Node.iSurf).PolyFlags |= PF_Selected;
+				FBSPOps::GErrors++;
+				//UE_LOG(LogEditorBsp, Warning, "Failed to link side" );
+			}
+
+			// Go to next side.
+			SkipSide:;
+		}
+	}
+	// Gather stats.
+	i=0; int j=0;
+	for( int32 iNode=0; iNode < Model->Nodes.Num(); iNode++ )
+	{
+		FBspNode&	Node		= Model->Nodes[iNode];
+		FVert*		VertPool	= &Model->Verts[Node.iVertPool];
+		
+		i += Node.NumVertices;
+		for( uint8 ThisVert=0; ThisVert < Node.NumVertices; ThisVert++ )
+			if( VertPool[ThisVert].iSide!=INDEX_NONE )
+				j++;
+	}
+	// Done.
+	//UE_LOG(LogEditorBsp, Log,  TEXT("BspOptGeom end") );
+// 	UE_LOG(LogEditorBsp, Log,  TEXT("Processed %i T-points, linked: %i/%i sides"), TeesFound, j, i );
+
+	PointVerts.Free();
+
+	// Remove unused vertices from the vertex streams.
+	// This is necessary to ensure the vertices added to eliminate T junctions
+	// don't overflow the 65536 vertex/stream limit.
+
+	FBSPOps::bspRefresh(Model,0);
+}
+
+
+void FBSPUtils::bspRepartition( UWorld* InWorld, int32 iNode )
+{
+	bspBuildFPolys( InWorld->GetModel(), 1, iNode );
+	bspMergeCoplanars( InWorld->GetModel(), 0, 0 );
+	FBSPOps::bspBuild( InWorld->GetModel(), FBSPOps::BSP_Good, 12, 70, 2, iNode );
+	FBSPOps::bspRefresh( InWorld->GetModel(), 1 );
+}
+
+
+void FBSPUtils::polySetAndClearPolyFlags(UModel *Model, uint32 SetBits, uint32 ClearBits,bool SelectedOnly, bool UpdateMaster)
+{
+	for( int32 i=0; i<Model->Surfs.Num(); i++ )
+	{
+		FBspSurf& Poly = Model->Surfs[i];
+		if( !SelectedOnly || (Poly.PolyFlags & PF_Selected) )
+		{
+			uint32 NewFlags = (Poly.PolyFlags & ~ClearBits) | SetBits;
+			if( NewFlags != Poly.PolyFlags )
+			{
+				Model->ModifySurf( i, UpdateMaster );
+				Poly.PolyFlags = NewFlags;
+				if (UpdateMaster)
+				{
+					const bool bUpdateTexCoords = false;
+					const bool bOnlyRefreshSurfaceMaterials = false;
+					polyUpdateMaster(Model, i, bUpdateTexCoords, bOnlyRefreshSurfaceMaterials);
+				}
+			}
+		}
+	}
+}
+
+
+bool FBSPUtils::polyFindMaster(UModel* InModel, int32 iSurf, FPoly &Poly)
+{
+	FBspSurf &Surf = InModel->Surfs[iSurf];
+	if( !Surf.Actor || !Surf.Actor->Brush->Polys->Element.IsValidIndex(Surf.iBrushPoly) )
+	{
+		return false;
+	}
+	else
+	{
+		Poly = Surf.Actor->Brush->Polys->Element[Surf.iBrushPoly];
+		return true;
+	}
+}
+
+
+void FBSPUtils::polyUpdateMaster
+(
+	UModel*	Model,
+	int32  	iSurf,
+	bool	bUpdateTexCoords,
+	bool	bOnlyRefreshSurfaceMaterials
+)
+{
+	FBspSurf &Surf = Model->Surfs[iSurf];
+	ABrush* Actor = Surf.Actor;
+	if( !Actor )
+		return;
+
+	UModel* Brush = Actor->Brush;
+	check(Brush);
+
+	FVector ActorLocation;
+	FVector ActorPrePivot;
+	FVector ActorScale;
+	FRotator ActorRotation;
+
+	if (Brush->bCachedOwnerTransformValid)
+	{
+		// Use transform cached when the geometry was last built, in case the current Actor transform has changed since then
+		// (e.g. because Auto Update BSP is disabled)
+		ActorLocation = Brush->OwnerLocationWhenLastBuilt;
+		ActorScale = Brush->OwnerScaleWhenLastBuilt;
+		ActorRotation = Brush->OwnerRotationWhenLastBuilt;
+	}
+	else
+	{
+		// No cached owner transform, so use the current one
+		ActorLocation = Actor->GetActorLocation();
+		ActorScale = Actor->GetActorScale();
+		ActorRotation = Actor->GetActorRotation();
+	}
+
+	const FRotationMatrix RotationMatrix(ActorRotation);
+
+	for (int32 iEdPoly = Surf.iBrushPoly; iEdPoly < Brush->Polys->Element.Num(); iEdPoly++)
+	{
+		FPoly& MasterEdPoly = Brush->Polys->Element[iEdPoly];
+		if (iEdPoly == Surf.iBrushPoly || MasterEdPoly.iLink == Surf.iBrushPoly)
+		{
+			MasterEdPoly.Material = Surf.Material;
+			MasterEdPoly.PolyFlags = Surf.PolyFlags & ~(PF_NoEdit);
+
+			if (bUpdateTexCoords)
+			{
+				MasterEdPoly.Base = RotationMatrix.InverseTransformVector(Model->Points[Surf.pBase] - ActorLocation) / ActorScale;
+				MasterEdPoly.TextureU = RotationMatrix.InverseTransformVector(Model->Vectors[Surf.vTextureU]) * ActorScale;
+				MasterEdPoly.TextureV = RotationMatrix.InverseTransformVector(Model->Vectors[Surf.vTextureV]) * ActorScale;
+			}
+		}
+	}
+
+	Model->InvalidSurfaces = true;
+
+	if (bOnlyRefreshSurfaceMaterials)
+	{
+		Model->bOnlyRebuildMaterialIndexBuffers = true;
+	}
+}
+
+
+void FBSPUtils::polyGetLinkedPolys
+(
+	ABrush* InBrush,
+	FPoly* InPoly,
+	TArray<FPoly>* InPolyList
+)
+{
+	InPolyList->Empty();
+
+	if( InPoly->iLink == INDEX_NONE )
+	{
+		// If this poly has no links, just stick the one poly in the final list.
+		new(*InPolyList)FPoly( *InPoly );
+	}
+	else
+	{
+		// Find all polys that match the source polys link value.
+		for( int32 poly = 0 ; poly < InBrush->Brush->Polys->Element.Num() ; poly++ )
+			if( InBrush->Brush->Polys->Element[poly].iLink == InPoly->iLink )
+				new(*InPolyList)FPoly( InBrush->Brush->Polys->Element[poly] );
+	}
+}
+
+
+void FBSPUtils::polySplitOverlappingEdges( TArray<FPoly>* InPolyList, TArray<FPoly>* InResult )
+{
+	InResult->Empty();
+
+	for( int32 poly = 0 ; poly < InPolyList->Num() ; poly++ )
+	{
+		FPoly* SrcPoly = &(*InPolyList)[poly];
+		FPoly NewPoly = *SrcPoly;
+
+		for( int32 edge = 0 ; edge < SrcPoly->Vertices.Num() ; edge++ )
+		{
+			FEdge SrcEdge = FEdge( SrcPoly->Vertices[edge], SrcPoly->Vertices[ edge+1 < SrcPoly->Vertices.Num() ? edge+1 : 0 ] );
+			FPlane SrcEdgePlane( SrcEdge.Vertex[0], SrcEdge.Vertex[1], SrcEdge.Vertex[0] + (SrcPoly->Normal * 16) );
+
+			for( int32 poly2 = 0 ; poly2 < InPolyList->Num() ; poly2++ )
+			{
+				FPoly* CmpPoly = &(*InPolyList)[poly2];
+
+				// We can't compare to ourselves.
+				if( CmpPoly == SrcPoly )
+					continue;
+
+				for( int32 edge2 = 0 ; edge2 < CmpPoly->Vertices.Num() ; edge2++ )
+				{
+					FEdge CmpEdge = FEdge( CmpPoly->Vertices[edge2], CmpPoly->Vertices[ edge2+1 < CmpPoly->Vertices.Num() ? edge2+1 : 0 ] );
+
+					// If both vertices on this edge lie on the same plane as the original edge, create
+					// a sphere around the original 2 vertices.  If either of this edges vertices are inside of
+					// that sphere, we need to split the original edge by adding a vertex to it's poly.
+					if( FMath::Abs( FVector::PointPlaneDist( CmpEdge.Vertex[0], SrcEdge.Vertex[0], SrcEdgePlane ) ) < THRESH_POINT_ON_PLANE
+						&& FMath::Abs( FVector::PointPlaneDist( CmpEdge.Vertex[1], SrcEdge.Vertex[0], SrcEdgePlane ) ) < THRESH_POINT_ON_PLANE )
+					{
+						//
+						// Check THIS edge against the SOURCE edge
+						//
+
+						FVector Dir = SrcEdge.Vertex[1] - SrcEdge.Vertex[0];
+						Dir.Normalize();
+						float Dist = FVector::Dist( SrcEdge.Vertex[1], SrcEdge.Vertex[0] );
+						FVector Origin = SrcEdge.Vertex[0] + (Dir * (Dist / 2.0f));
+						float Radius = Dist / 2.0f;
+
+						for( int32 vtx = 0 ; vtx < 2 ; vtx++ )
+							if( FVector::Dist( Origin, CmpEdge.Vertex[vtx] ) && FVector::Dist( Origin, CmpEdge.Vertex[vtx] ) < Radius )
+								NewPoly.InsertVertex( edge2+1, CmpEdge.Vertex[vtx] );
+					}
+				}
+			}
+		}
+
+		new(*InResult)FPoly( NewPoly );
+	}
+
+}
+
+
+void FBSPUtils::polyGetOuterEdgeList
+(
+	TArray<FPoly>* InPolyList,
+	TArray<FEdge>* InEdgeList
+)
+{
+	TArray<FPoly> NewPolyList;
+	polySplitOverlappingEdges( InPolyList, &NewPolyList );
+
+	TArray<FEdge> TempEdges;
+
+	// Create a master list of edges.
+	for( int32 poly = 0 ; poly < NewPolyList.Num() ; poly++ )
+	{
+		FPoly* Poly = &NewPolyList[poly];
+		for( int32 vtx = 0 ; vtx < Poly->Vertices.Num() ; vtx++ )
+			new( TempEdges )FEdge( Poly->Vertices[vtx], Poly->Vertices[ vtx+1 < Poly->Vertices.Num() ? vtx+1 : 0] );
+	}
+
+	// Add all the unique edges into the final edge list.
+	TArray<FEdge> FinalEdges;
+
+	for( int32 tedge = 0 ; tedge < TempEdges.Num() ; tedge++ )
+	{
+		FEdge* TestEdge = &TempEdges[tedge];
+
+		int32 EdgeCount = 0;
+		for( int32 edge = 0 ; edge < TempEdges.Num() ; edge++ )
+		{
+			if( TempEdges[edge] == *TestEdge )
+				EdgeCount++;
+		}
+
+		if( EdgeCount == 1 )
+			new( FinalEdges )FEdge( *TestEdge );
+	}
+
+	// Reorder all the edges so that they line up, end to end.
+	InEdgeList->Empty();
+	if( !FinalEdges.Num() ) return;
+
+	new( *InEdgeList )FEdge( FinalEdges[0] );
+	FVector Comp = FinalEdges[0].Vertex[1];
+	FinalEdges.RemoveAt(0);
+
+	FEdge DebuG;
+	for( int32 x = 0 ; x < FinalEdges.Num() ; x++ )
+	{
+		DebuG = FinalEdges[x];
+
+		// If the edge is backwards, flip it
+		if( FinalEdges[x].Vertex[1] == Comp )
+			Exchange( FinalEdges[x].Vertex[0], FinalEdges[x].Vertex[1] );
+
+		if( FinalEdges[x].Vertex[0] == Comp )
+		{
+			new( *InEdgeList )FEdge( FinalEdges[x] );
+			Comp = FinalEdges[x].Vertex[1];
+			FinalEdges.RemoveAt(x);
+			x = -1;
+		}
+	}
+
+}
+
diff --git a/Engine/Source/Developer/BSPUtils/Public/BSPOps.h b/Engine/Source/Developer/BSPUtils/Public/BSPOps.h
new file mode 100644
--- /dev/null
+++ b/Engine/Source/Developer/BSPUtils/Public/BSPOps.h
@@ -0,0 +1,150 @@
+// Copyright Epic Games, Inc. All Rights Reserved.
+
+#pragma once
+
+#include "CoreMinimal.h"
+#include "Engine/Brush.h"
+#include "Engine/Polys.h"
+
+class AVolume;
+class UModel;
+
+class FBSPOps
+{
+public:
+	/** Quality level for rebuilding Bsp. */
+	enum EBspOptimization
+	{
+		BSP_Lame,
+		BSP_Good,
+		BSP_Optimal
+	};
+
+	/** Possible positions of a child Bsp node relative to its parent (for BspAddToNode) */
+	enum ENodePlace 
+	{
+		NODE_Back		= 0, // Node is in back of parent              -> Bsp[iParent].iBack.
+		NODE_Front		= 1, // Node is in front of parent             -> Bsp[iParent].iFront.
+		NODE_Plane		= 2, // Node is coplanar with parent           -> Bsp[iParent].iPlane.
+		NODE_Root		= 3, // Node is the Bsp root and has no parent -> Bsp[0].
+	};
+
+	BSPUTILS_API static void csgPrepMovingBrush( ABrush* Actor );
+	BSPUTILS_API static void csgCopyBrush( ABrush* Dest, ABrush* Src, uint32 PolyFlags, EObjectFlags ResFlags, bool bNeedsPrep, bool bCopyPosRotScale, bool bAllowEmpty = false );
+	BSPUTILS_API static ABrush*	csgAddOperation( ABrush* Actor, uint32 PolyFlags, EBrushType BrushType );
+	BSPUTILS_API static int32 bspAddVector( UModel* Model, const FVector* V, bool Exact );
+	BSPUTILS_API static int32 bspAddPoint( UModel* Model, const FVector* V, bool Exact );
+	BSPUTILS_API static void bspBuild( UModel* Model, enum EBspOptimization Opt, int32 Balance, int32 PortalBias, int32 RebuildSimplePolys, int32 iNode );
+	BSPUTILS_API static void bspRefresh( UModel* Model, bool NoRemapSurfs );
+	BSPUTILS_API static void bspBuildBounds( UModel* Model );
+	BSPUTILS_API static void bspValidateBrush( UModel* Brush, bool ForceValidate, bool DoStatusUpdate );
+	BSPUTILS_API static void bspUnlinkPolys( UModel* Brush );
+	BSPUTILS_API static int32 bspAddNode( UModel* Model, int32 iParent, enum ENodePlace ENodePlace, uint32 NodeFlags, FPoly* EdPoly );
+
+	/**
+	 * Rebuild some brush internals
+	 */
+	BSPUTILS_API static void RebuildBrush(UModel* Brush);
+
+	BSPUTILS_API static FPoly BuildInfiniteFPoly( UModel* Model, int32 iNode );
+
+	/** Called when an AVolume shape is changed*/
+	BSPUTILS_API static void HandleVolumeShapeChanged(AVolume& Volume);
+
+	/** Errors encountered in Csg operation. */
+	BSPUTILS_API static int32 GErrors;
+	BSPUTILS_API static bool GFastRebuild;
+
+protected:
+	/**
+	 * Rotates the specified brush's vertices.
+	 */
+	static void RotateBrushVerts(ABrush* Brush, const FRotator& Rotation, bool bClearComponents);
+
+	static void SplitPolyList
+	(
+		UModel				*Model,
+		int32               iParent,
+		ENodePlace			NodePlace,
+		int32               NumPolys,
+		FPoly				**PolyList,
+		EBspOptimization	Opt,
+		int32				Balance,
+		int32				PortalBias,
+		int32				RebuildSimplePolys
+	);
+};
+
+/** @todo: remove when uses of ENodePlane have been replaced with FBSPOps::ENodePlace. */
+typedef FBSPOps::ENodePlace ENodePlace;
+
+
+struct FBspPointsKey
+{
+	int32 X;
+	int32 Y;
+	int32 Z;
+
+	FBspPointsKey(int32 InX, int32 InY, int32 InZ)
+		: X(InX)
+		, Y(InY)
+		, Z(InZ)
+	{}
+
+	friend FORCEINLINE bool operator == (const FBspPointsKey& A, const FBspPointsKey& B)
+	{
+		return A.X == B.X && A.Y == B.Y && A.Z == B.Z;
+	}
+
+	friend FORCEINLINE uint32 GetTypeHash(const FBspPointsKey& Key)
+	{
+		return HashCombine(static_cast<uint32>(Key.X), HashCombine(static_cast<uint32>(Key.Y), static_cast<uint32>(Key.Z)));
+	}
+};
+
+struct FBspIndexedPoint
+{
+	FBspIndexedPoint(const FVector& InPoint, int32 InIndex)
+		: Point(InPoint)
+		, Index(InIndex)
+	{}
+
+	FVector Point;
+	int32 Index;
+};
+
+
+struct FBspPointsGridItem
+{
+	TArray<FBspIndexedPoint, TInlineAllocator<16>> IndexedPoints;
+};
+
+
+// Represents a sparse granular 3D grid into which points are added for quick (~O(1)) lookup.
+// The 3D space is divided into a grid with a given granularity.
+// Points are considered to have a given radius (threshold) and are added to the grid cube they fall in, and to up to seven neighbours if they overlap.
+class FBspPointsGrid
+{
+public:
+	FBspPointsGrid(float InGranularity, float InThreshold, int32 InitialSize = 0)
+		: OneOverGranularity(1.0f / InGranularity)
+		, Threshold(InThreshold)
+	{
+		check(InThreshold / InGranularity <= 0.5f);
+		Clear(InitialSize);
+	}
+
+	BSPUTILS_API void Clear(int32 InitialSize = 0);
+
+	BSPUTILS_API int32 FindOrAddPoint(const FVector& Point, int32 Index, float Threshold);
+
+	BSPUTILS_API static FBspPointsGrid* GBspPoints;
+	BSPUTILS_API static FBspPointsGrid* GBspVectors;
+
+private:
+	float OneOverGranularity;
+	float Threshold;
+
+	typedef TMap<FBspPointsKey, FBspPointsGridItem> FGridMap;
+	FGridMap GridMap;
+};
diff --git a/Engine/Source/Developer/BSPUtils/Public/BSPUtils.h b/Engine/Source/Developer/BSPUtils/Public/BSPUtils.h
new file mode 100644
--- /dev/null
+++ b/Engine/Source/Developer/BSPUtils/Public/BSPUtils.h
@@ -0,0 +1,93 @@
+// Copyright Epic Games, Inc. All Rights Reserved.
+
+#pragma once
+
+#include "CoreMinimal.h"
+#include "Engine/Brush.h"
+
+class UWorld;
+class UModel;
+class UMaterialInterface;
+class FPoly;
+
+
+class BSPUTILS_API FBSPUtils
+{
+public:
+	/** Repartition Bsp tree */
+	static void bspRepartition( UWorld* InWorld, int32 iNode );
+
+	/** Convert a Bsp node to an EdPoly.  Returns number of vertices in Bsp node. */
+	static int32 bspNodeToFPoly( UModel* Model, int32 iNode, FPoly* EdPoly );
+
+	/**
+	 * Clean up all nodes after a CSG operation.  Resets temporary bit flags and unlinks
+	 * empty leaves.  Removes zero-vertex nodes which have nonzero-vertex coplanars.
+	 */
+	static void bspCleanup( UModel* Model );
+
+	/** 
+	 * Build EdPoly list from a model's Bsp. Not transactional.
+	 * @param DestArray helps build bsp FPolys in non-main threads. It also allows to perform this action without GUndo 
+	 *	      interfering. Temporary results will be written to DestArray. Defaults to Model->Polys->Element
+	 */
+	static void bspBuildFPolys( UModel* Model, bool SurfLinks, int32 iNode, TArray<FPoly>* DestArray = NULL );
+
+	static void bspMergeCoplanars( UModel* Model, bool RemapLinks, bool MergeDisparateTextures );
+
+	/**
+	 * Performs any CSG operation between the brush and the world.
+	 */
+	static int32 bspBrushCSG( ABrush* Actor, UModel* Model, UModel* TempModel, UMaterialInterface* SelectedMaterialInstance, uint32 PolyFlags, EBrushType BrushType, ECsgOper CSGOper, bool bBuildBounds, bool bMergePolys, bool bReplaceNULLMaterialRefs, bool bShowProgressBar=true );
+
+	/**
+	 * Optimize a level's Bsp, eliminating T-joints where possible, and building side
+	 * links.  This does not always do a 100% perfect job, mainly due to imperfect 
+	 * levels, however it should never fail or return incorrect results.
+	 */
+	static void bspOptGeom( UModel* Model );
+
+	/**
+	 * Sets and clears all Bsp node flags.  Affects all nodes, even ones that don't
+	 * really exist.
+	 */
+	static void polySetAndClearPolyFlags(UModel *Model, uint32 SetBits, uint32 ClearBits,bool SelectedOnly, bool UpdateMaster);
+
+	/**
+	 *
+	 * Find the Brush EdPoly corresponding to a given Bsp surface.
+	 *
+	 * @param InModel	Model to get poly from
+	 * @param iSurf		surface index
+	 * @param Poly		
+	 *
+	 * returns true if poly not available
+	 */
+	static bool polyFindMaster(UModel* InModel, int32 iSurf, FPoly &Poly);
+
+	/**
+	 * Update a the master brush EdPoly corresponding to a newly-changed
+	 * poly to reflect its new properties.
+	 *
+	 * Doesn't do any transaction tracking.
+	 */
+	static void polyUpdateMaster(UModel* Model, int32 iSurf, bool bUpdateTexCoords, bool bOnlyRefreshSurfaceMaterials);
+	
+	/**
+	 * Populates a list with all polys that are linked to the specified poly.  The
+	 * resulting list includes the original poly.
+	 */
+	static void polyGetLinkedPolys(ABrush* InBrush, FPoly* InPoly, TArray<FPoly>* InPolyList);
+
+	/**
+	 * Takes a list of polygons and creates a new list of polys which have no overlapping edges.  It splits
+	 * edges as necessary to achieve this.
+	 */
+	static void polySplitOverlappingEdges( TArray<FPoly>* InPolyList, TArray<FPoly>* InResult );
+
+	/**
+	 * Takes a list of polygons and returns a list of the outside edges (edges which are not shared
+	 * by other polys in the list).
+	 */
+	static void polyGetOuterEdgeList(TArray<FPoly>* InPolyList, TArray<FEdge>* InEdgeList);
+};
diff --git a/Engine/Source/Developer/HierarchicalLODUtilities/HierarchicalLODUtilities.Build.cs b/Engine/Source/Developer/HierarchicalLODUtilities/HierarchicalLODUtilities.Build.cs
--- a/Engine/Source/Developer/HierarchicalLODUtilities/HierarchicalLODUtilities.Build.cs
+++ b/Engine/Source/Developer/HierarchicalLODUtilities/HierarchicalLODUtilities.Build.cs
@@ -18,6 +18,7 @@ public class HierarchicalLODUtilities : ModuleRules
         PrivateDependencyModuleNames.AddRange(
             new string[]
 			{
+				"BSPUtils",
 				"EditorFramework",
 				"Engine",
 				"MeshDescription",
diff --git a/Engine/Source/Editor/UnrealEd/Private/BSPOps.cpp b/Engine/Source/Editor/UnrealEd/Private/BSPOps.cpp
deleted file mode 100644
diff --git a/Engine/Source/Editor/UnrealEd/Private/EditorActor.cpp b/Engine/Source/Editor/UnrealEd/Private/EditorActor.cpp
--- a/Engine/Source/Editor/UnrealEd/Private/EditorActor.cpp
+++ b/Engine/Source/Editor/UnrealEd/Private/EditorActor.cpp
@@ -46,6 +46,7 @@
 #include "Engine/LevelStreaming.h"
 #include "LevelUtils.h"
 #include "BusyCursor.h"
+#include "Engine/BrushBuilder.h"
 #include "BSPOps.h"
 #include "EditorLevelUtils.h"
 #include "Kismet2/BlueprintEditorUtils.h"
@@ -1289,6 +1290,11 @@ void UUnrealEdEngine::edactReplaceSelectedBrush( UWorld* InWorld )
 		ABrush* NewBrush = FBSPOps::csgAddOperation( DefaultBrush, SrcBrush->PolyFlags, (EBrushType)SrcBrush->BrushType );
 		if( NewBrush )
 		{
+			if( NewBrush->GetBrushBuilder() )
+			{
+				FActorLabelUtilities::SetActorLabelUnique(NewBrush, FText::Format(NSLOCTEXT("UnrealEd", "BrushName", "{0} Brush"), FText::FromString(NewBrush->GetBrushBuilder()->GetClass()->GetDescription())).ToString());
+			}
+
 			SrcBrush->MarkPackageDirty();
 			NewBrush->MarkPackageDirty();
 
diff --git a/Engine/Source/Editor/UnrealEd/Private/EditorBsp.cpp b/Engine/Source/Editor/UnrealEd/Private/EditorBsp.cpp
--- a/Engine/Source/Editor/UnrealEd/Private/EditorBsp.cpp
+++ b/Engine/Source/Editor/UnrealEd/Private/EditorBsp.cpp
@@ -5,909 +5,30 @@
 =============================================================================*/
 
 #include "CoreMinimal.h"
-#include "Misc/MemStack.h"
-#include "Misc/FeedbackContext.h"
 #include "Materials/MaterialInterface.h"
 #include "Model.h"
-#include "Materials/Material.h"
 #include "Editor/EditorEngine.h"
-#include "Engine/Polys.h"
 #include "Editor.h"
 #include "BSPOps.h"
+#include "BSPUtils.h"
 #include "Engine/Selection.h"
 
-DEFINE_LOG_CATEGORY_STATIC(LogEditorBsp, Log, All);
-
-/*---------------------------------------------------------------------------------------
-	Globals.
----------------------------------------------------------------------------------------*/
-
-// Magic numbers.
-#define THRESH_OPTGEOM_COPLANAR			(0.25)		/* Threshold for Bsp geometry optimization */
-#define THRESH_OPTGEOM_COSIDAL			(0.25)		/* Threshold for Bsp geometry optimization */
-
-//
-// Status of filtered polygons:
-//
-enum EPolyNodeFilter
-{
-	F_OUTSIDE				= 0, // Leaf is an exterior leaf (visible to viewers).
-	F_INSIDE				= 1, // Leaf is an interior leaf (non-visible, hidden behind backface).
-	F_COPLANAR_OUTSIDE		= 2, // Poly is coplanar and in the exterior (visible to viewers).
-	F_COPLANAR_INSIDE		= 3, // Poly is coplanar and inside (invisible to viewers).
-	F_COSPATIAL_FACING_IN	= 4, // Poly is coplanar, cospatial, and facing in.
-	F_COSPATIAL_FACING_OUT	= 5, // Poly is coplanar, cospatial, and facing out.
-};
-
-//
-// Generic filter function called by BspFilterEdPolys.  A and B are pointers
-// to any integers that your specific routine requires (or NULL if not needed).
-//
-typedef void (*BSP_FILTER_FUNC)
-(
-	UModel			*Model,
-	int32			    iNode,
-	FPoly			*EdPoly,
-	EPolyNodeFilter Leaf,
-	ENodePlace	ENodePlace
-);
-
-//
-// Information used by FilterEdPoly.
-//
-class FCoplanarInfo
-{
-public:
-	int32	    iOriginalNode;
-	int32     iBackNode;
-	int	    BackNodeOutside;
-	int	    FrontLeafOutside;
-	int     ProcessingBack;
-};
-
-//
-// Function to filter an EdPoly through the Bsp, calling a callback
-// function for all chunks that fall into leaves.
-//
-void FilterEdPoly
-(
-	BSP_FILTER_FUNC FilterFunc, 
-	UModel			*Model,
-	int32  			iNode, 
-	FPoly			*EdPoly, 
-	FCoplanarInfo	CoplanarInfo, 
-	int				Outside
-);
-
-//
-// Bsp statistics used by link topic function.
-//
-class FBspStats
-{
-public:
-	int   Polys;      // Number of BspSurfs.
-	int   Nodes;      // Total number of Bsp nodes.
-	int   MaxDepth;   // Maximum tree depth.
-	int   AvgDepth;   // Average tree depth.
-	int   Branches;   // Number of nodes with two children.
-	int   Coplanars;  // Number of nodes holding coplanar polygons (in same plane as parent).
-	int   Fronts;     // Number of nodes with only a front child.
-	int   Backs;      // Number of nodes with only a back child.
-	int   Leaves;     // Number of nodes with no children.
-	int   FrontLeaves;// Number of leaf nodes that are in front of their parent.
-	int   BackLeaves; // Number of leaf nodes that are in back of their parent.
-	int   DepthCount; // Depth counter (used for finding average depth).
-} GBspStats;
-
-//
-// Global variables shared between bspBrushCSG and AddWorldToBrushFunc.  These are very
-// tightly tied into the function AddWorldToBrush, not general-purpose.
-//
-int32			GDiscarded;		// Number of polys discarded and not added.
-int32         GNode;          // Node AddBrushToWorld is adding to.
-int32         GLastCoplanar;	// Last coplanar beneath GNode at start of AddWorldToBrush.
-int32         GNumNodes;		// Number of Bsp nodes at start of AddWorldToBrush.
-UModel		*GModel;		// Level map Model we're adding to.
-
-/*----------------------------------------------------------------------------
-   EdPoly building and compacting.
-----------------------------------------------------------------------------*/
-
-//
-// Trys to merge two polygons.  If they can be merged, replaces Poly1 and emptys Poly2
-// and returns 1.  Otherwise, returns 0.
-//
-int TryToMerge( FPoly *Poly1, FPoly *Poly2 )
-{
-	// Find one overlapping point.
-	int32 Start1=0, Start2=0;
-	for( Start1=0; Start1<Poly1->Vertices.Num(); Start1++ )
-		for( Start2=0; Start2<Poly2->Vertices.Num(); Start2++ )
-			if( FVector::PointsAreSame(Poly1->Vertices[Start1], Poly2->Vertices[Start2]) )
-				goto FoundOverlap;
-	return 0;
-	FoundOverlap:
-
-	// Wrap around trying to merge.
-	int32 End1  = Start1;
-	int32 End2  = Start2;
-	int32 Test1 = Start1+1; if (Test1>=Poly1->Vertices.Num()) Test1 = 0;
-	int32 Test2 = Start2-1; if (Test2<0)                   Test2 = Poly2->Vertices.Num()-1;
-	if( FVector::PointsAreSame(Poly1->Vertices[Test1],Poly2->Vertices[Test2]) )
-	{
-		End1   = Test1;
-		Start2 = Test2;
-	}
-	else
-	{
-		Test1 = Start1-1; if (Test1<0)                   Test1=Poly1->Vertices.Num()-1;
-		Test2 = Start2+1; if (Test2>=Poly2->Vertices.Num()) Test2=0;
-		if( FVector::PointsAreSame(Poly1->Vertices[Test1],Poly2->Vertices[Test2]) )
-		{
-			Start1 = Test1;
-			End2   = Test2;
-		}
-		else return 0;
-	}
-
-	// Build a new edpoly containing both polygons merged.
-	FPoly NewPoly = *Poly1;
-	NewPoly.Vertices.Empty();
-	int32 Vertex = End1;
-	for( int32 i=0; i<Poly1->Vertices.Num(); i++ )
-	{
-		new(NewPoly.Vertices) FVector3f(Poly1->Vertices[Vertex]);
-		if( ++Vertex >= Poly1->Vertices.Num() )
-			Vertex=0;
-	}
-	Vertex = End2;
-	for( int32 i=0; i<(Poly2->Vertices.Num()-2); i++ )
-	{
-		if( ++Vertex >= Poly2->Vertices.Num() )
-			Vertex=0;
-		new(NewPoly.Vertices) FVector3f(Poly2->Vertices[Vertex]);
-	}
-
-	// Remove colinear vertices and check convexity.
-	if( NewPoly.RemoveColinears() )
-	{
-		*Poly1 = NewPoly;
-		Poly2->Vertices.Empty();
-		return true;
-	}
-	else return 0;
-}
-
-//
-// Merge all polygons in coplanar list that can be merged convexly.
-//
-void MergeCoplanars( UModel* Model, int32* PolyList, int32 PolyCount )
-{
-	int32 MergeAgain = 1;
-	while( MergeAgain )
-	{
-		MergeAgain = 0;
-		for( int32 i=0; i<PolyCount; i++ )
-		{
-			FPoly& Poly1 = Model->Polys->Element[PolyList[i]];
-			if( Poly1.Vertices.Num() > 0 )
-	        {
-				for( int32 j=i+1; j<PolyCount; j++ )
-				{
-					FPoly& Poly2 = Model->Polys->Element[PolyList[j]];
-					if( Poly2.Vertices.Num() > 0 )
-					{
-                  		if( TryToMerge( &Poly1, &Poly2 ) )
-							MergeAgain=1;
-					}
-				}
-			}
-		}
-	}
-}
-
-//
-// Convert a Bsp node's polygon to an EdPoly, add it to the list, and recurse.
-//
-void MakeEdPolys( UModel* Model, int32 iNode, TArray<FPoly>* DestArray )
-{
-	FBspNode* Node = &Model->Nodes[iNode];
-
-	FPoly Temp;
-	if( GEditor->bspNodeToFPoly(Model,iNode,&Temp) >= 3 )
-	{
-		new(*DestArray)FPoly(Temp);
-	}
-
-	if( Node->iFront!=INDEX_NONE )
-	{
-		MakeEdPolys( Model, Node->iFront, DestArray );
-	}
-	if( Node->iBack !=INDEX_NONE ) 
-	{
-		MakeEdPolys( Model, Node->iBack, DestArray );
-	}
-	if( Node->iPlane!=INDEX_NONE )
-	{
-		MakeEdPolys( Model, Node->iPlane, DestArray );
-	}
-}
 
 void UEditorEngine::bspBuildFPolys( UModel* Model, bool SurfLinks, int32 iNode, TArray<FPoly>* DestArray )
 {
-	if (DestArray == NULL)
-	{
-		DestArray = &Model->Polys->Element;
-	}
-	DestArray->Reset();
-
-	if( Model->Nodes.Num() )
-	{
-		MakeEdPolys( Model, iNode, DestArray );
-	}
-
-	if( !SurfLinks )
-	{
-		const int32 ElementsCount = DestArray->Num();
-		for( int32 i=0; i<ElementsCount; i++ )
-		{
-			(*DestArray)[i].iLink=i;
-		}
-	}
+	FBSPUtils::bspBuildFPolys(Model, SurfLinks, iNode, DestArray);
 }
 
-void UEditorEngine::bspMergeCoplanars( UModel* Model, bool RemapLinks, bool MergeDisparateTextures )
-{
-	int32 OriginalNum = Model->Polys->Element.Num();
-
-	// Mark all polys as unprocessed.
-	for( int32 i=0; i<Model->Polys->Element.Num(); i++ )
-		Model->Polys->Element[i].PolyFlags &= ~PF_EdProcessed;
-
-	// Find matching coplanars and merge them.
-	FMemMark Mark(FMemStack::Get());
-	int32* PolyList = new(FMemStack::Get(),Model->Polys->Element.Num())int32;
-	int32 n=0;
-	for( int32 i=0; i<Model->Polys->Element.Num(); i++ )
-	{
-		FPoly* EdPoly = &Model->Polys->Element[i];
-		if( EdPoly->Vertices.Num()>0 && !(EdPoly->PolyFlags & PF_EdProcessed) )
-		{
-			int32 PolyCount         =  0;
-			PolyList[PolyCount++] =  i;
-			EdPoly->PolyFlags    |= PF_EdProcessed;
-			for( int32 j=i+1; j<Model->Polys->Element.Num(); j++ )
-	        {
-				FPoly* OtherPoly = &Model->Polys->Element[j];
-				if( OtherPoly->iLink == EdPoly->iLink && OtherPoly->Vertices.Num() )
-				{
-					float Dist = (OtherPoly->Vertices[0] - EdPoly->Vertices[0]) | EdPoly->Normal;
-					if
-					(	Dist>-0.001
-					&&	Dist<0.001
-					&&	(OtherPoly->Normal|EdPoly->Normal)>0.9999
-					&&	(MergeDisparateTextures
-						||	(	FVector::PointsAreNear(OtherPoly->TextureU,EdPoly->TextureU,THRESH_VECTORS_ARE_NEAR)
-							&&	FVector::PointsAreNear(OtherPoly->TextureV,EdPoly->TextureV,THRESH_VECTORS_ARE_NEAR) ) ) )
-					{
-						OtherPoly->PolyFlags |= PF_EdProcessed;
-						PolyList[PolyCount++] = j;
-					}
-				}
-			}
-			if( PolyCount > 1 )
-	        {
-				MergeCoplanars( Model, PolyList, PolyCount );
-				n++;
-			}
-		}
-	}
-//	UE_LOG(LogEditorBsp, Log,  TEXT("Found %i coplanar sets in %i"), n, Model->Polys->Element.Num() );
-	Mark.Pop();
-
-	// Get rid of empty EdPolys while remapping iLinks.
-	FMemMark Mark2(FMemStack::Get());
-	int32 j=0;
-	int32* Remap = new(FMemStack::Get(),Model->Polys->Element.Num())int32;
-	for( int32 i=0; i<Model->Polys->Element.Num(); i++ )
-	{
-		if( Model->Polys->Element[i].Vertices.Num() )
-		{
-			Remap[i] = j;
-			Model->Polys->Element[j] = Model->Polys->Element[i];
-			j++;
-		}
-	}
-	Model->Polys->Element.RemoveAt( j, Model->Polys->Element.Num()-j );
-	if( RemapLinks )
-	{
-		for( int32 i=0; i<Model->Polys->Element.Num(); i++ )
-		{
-			if (Model->Polys->Element[i].iLink != INDEX_NONE)
-			{
-				CA_SUPPRESS(6001); // warning C6001: Using uninitialized memory '*Remap'.
-				Model->Polys->Element[i].iLink = Remap[Model->Polys->Element[i].iLink];
-			}
-		}
-	}
-// 	UE_LOG(LogEditorBsp, Log,  TEXT("BspMergeCoplanars reduced %i->%i"), OriginalNum, Model->Polys->Element.Num() );
-	Mark2.Pop();
-}
-
-/*----------------------------------------------------------------------------
-   CSG types & general-purpose callbacks.
-----------------------------------------------------------------------------*/
 
-//
-// Recursive worker function called by BspCleanup.
-//
-void CleanupNodes( UModel *Model, int32 iNode, int32 iParent )
+void UEditorEngine::bspMergeCoplanars( UModel* Model, bool RemapLinks, bool MergeDisparateTextures )
 {
-	FBspNode *Node = &Model->Nodes[iNode];
-
-	// Transactionally empty vertices of tag-for-empty nodes.
-	Node->NodeFlags &= ~(NF_IsNew | NF_IsFront | NF_IsBack);
-
-	// Recursively clean up front, back, and plane nodes.
-	if( Node->iFront != INDEX_NONE ) CleanupNodes( Model, Node->iFront, iNode );
-	if( Node->iBack  != INDEX_NONE ) CleanupNodes( Model, Node->iBack , iNode );
-	if( Node->iPlane != INDEX_NONE ) CleanupNodes( Model, Node->iPlane, iNode );
-
-	// Reload Node since the recusive call aliases it.
-	Node = &Model->Nodes[iNode];
-
-	// If this is an empty node with a coplanar, replace it with the coplanar.
-	if( Node->NumVertices==0 && Node->iPlane!=INDEX_NONE )
-	{
-		FBspNode* PlaneNode = &Model->Nodes[ Node->iPlane ];
-
-		// Stick our front, back, and parent nodes on the coplanar.
-		if( (Node->Plane | PlaneNode->Plane) >= 0.0 )
-		{
-			PlaneNode->iFront  = Node->iFront;
-			PlaneNode->iBack   = Node->iBack;
-		}
-		else
-		{
-			PlaneNode->iFront  = Node->iBack;
-			PlaneNode->iBack   = Node->iFront;
-		}
-
-		if( iParent == INDEX_NONE )
-		{
-			// This node is the root.
-			*Node                  = *PlaneNode;   // Replace root.
-			PlaneNode->NumVertices = 0;            // Mark as unused.
-		}
-		else
-		{
-			// This is a child node.
-			FBspNode *ParentNode = &Model->Nodes[iParent];
-
-			if      ( ParentNode->iFront == iNode ) ParentNode->iFront = Node->iPlane;
-			else if ( ParentNode->iBack  == iNode ) ParentNode->iBack  = Node->iPlane;
-			else if ( ParentNode->iPlane == iNode ) ParentNode->iPlane = Node->iPlane;
-			else UE_LOG(LogEditorBsp, Fatal, TEXT("CleanupNodes: Parent and child are unlinked"));
-		}
-	}
-	else if( Node->NumVertices == 0 && ( Node->iFront==INDEX_NONE || Node->iBack==INDEX_NONE ) )
-	{
-		// Delete empty nodes with no fronts or backs.
-		// Replace empty nodes with only fronts.
-		// Replace empty nodes with only backs.
-		int32 iReplacementNode;
-		if     ( Node->iFront != INDEX_NONE ) iReplacementNode = Node->iFront;
-		else if( Node->iBack  != INDEX_NONE ) iReplacementNode = Node->iBack;
-		else                                  iReplacementNode = INDEX_NONE;
-
-		if( iParent == INDEX_NONE )
-		{
-			// Root.
-			if( iReplacementNode == INDEX_NONE )
-			{
-         		Model->Nodes.Empty();
-			}
-			else
-			{
-         		*Node = Model->Nodes[iReplacementNode];
-			}
-		}
-		else
-		{
-			// Regular node.
-			FBspNode *ParentNode = &Model->Nodes[iParent];
-
-			if     ( ParentNode->iFront == iNode ) ParentNode->iFront = iReplacementNode;
-			else if( ParentNode->iBack  == iNode ) ParentNode->iBack  = iReplacementNode;
-			else if( ParentNode->iPlane == iNode ) ParentNode->iPlane = iReplacementNode;
-			else UE_LOG(LogEditorBsp, Fatal, TEXT("CleanupNodes: Parent and child are unlinked"));
-		}
-	}
+	FBSPUtils::bspMergeCoplanars(Model, RemapLinks, MergeDisparateTextures);
 }
 
 
 void UEditorEngine::bspCleanup( UModel *Model )
 {
-	if( Model->Nodes.Num() > 0 ) 
-		CleanupNodes( Model, 0, INDEX_NONE );
-}
-
-/*----------------------------------------------------------------------------
-   CSG leaf filter callbacks.
-----------------------------------------------------------------------------*/
-
-void AddBrushToWorldFunc( UModel* Model, int32 iNode, FPoly* EdPoly,
-	EPolyNodeFilter Filter, ENodePlace ENodePlace )
-{
-	switch( Filter )
-	{
-		case F_OUTSIDE:
-		case F_COPLANAR_OUTSIDE:
-			FBSPOps::bspAddNode (Model,iNode,ENodePlace,NF_IsNew,EdPoly);
-			break;
-		case F_COSPATIAL_FACING_OUT:
-			if( !(EdPoly->PolyFlags & PF_Semisolid) )
-				FBSPOps::bspAddNode (Model,iNode,ENodePlace,NF_IsNew,EdPoly);
-			break;
-		case F_INSIDE:
-		case F_COPLANAR_INSIDE:
-		case F_COSPATIAL_FACING_IN:
-			break;
-	}
-}
-
-void AddWorldToBrushFunc( UModel* Model, int32 iNode, FPoly* EdPoly,
-	EPolyNodeFilter Filter, ENodePlace ENodePlace )
-{
-	switch( Filter )
-	{
-		case F_OUTSIDE:
-		case F_COPLANAR_OUTSIDE:
-			// Only affect the world poly if it has been cut.
-			if( EdPoly->PolyFlags & PF_EdCut )
-				FBSPOps::bspAddNode( GModel, GLastCoplanar, FBSPOps::NODE_Plane, NF_IsNew, EdPoly );
-			break;
-		case F_INSIDE:
-		case F_COPLANAR_INSIDE:
-		case F_COSPATIAL_FACING_IN:
-		case F_COSPATIAL_FACING_OUT:
-			// Discard original poly.
-			GDiscarded++;
-			if( GModel->Nodes[GNode].NumVertices )
-			{
-				GModel->Nodes[GNode].NumVertices = 0;
-			}
-			break;
-	}
-}
-
-void SubtractBrushFromWorldFunc( UModel* Model, int32 iNode, FPoly* EdPoly,
-	EPolyNodeFilter Filter, ENodePlace ENodePlace )
-{
-	switch (Filter)
-	{
-		case F_OUTSIDE:
-		case F_COSPATIAL_FACING_OUT:
-		case F_COSPATIAL_FACING_IN:
-		case F_COPLANAR_OUTSIDE:
-			break;
-		case F_COPLANAR_INSIDE:
-		case F_INSIDE:
-			EdPoly->Reverse();
-			FBSPOps::bspAddNode (Model,iNode,ENodePlace,NF_IsNew,EdPoly); // Add to Bsp back
-			EdPoly->Reverse();
-			break;
-	}
-}
-
-void SubtractWorldToBrushFunc( UModel* Model, int32 iNode, FPoly* EdPoly,
-	EPolyNodeFilter Filter, ENodePlace ENodePlace )
-{
-	switch( Filter )
-	{
-		case F_OUTSIDE:
-		case F_COPLANAR_OUTSIDE:
-		case F_COSPATIAL_FACING_IN:
-			// Only affect the world poly if it has been cut.
-			if( EdPoly->PolyFlags & PF_EdCut )
-				FBSPOps::bspAddNode( GModel, GLastCoplanar, FBSPOps::NODE_Plane, NF_IsNew, EdPoly );
-			break;
-		case F_INSIDE:
-		case F_COPLANAR_INSIDE:
-		case F_COSPATIAL_FACING_OUT:
-			// Discard original poly.
-			GDiscarded++;
-			if( GModel->Nodes[GNode].NumVertices )
-			{
-				GModel->Nodes[GNode].NumVertices = 0;
-			}
-			break;
-	}
-}
-
-void IntersectBrushWithWorldFunc( UModel* Model, int32 iNode, FPoly *EdPoly,
-	EPolyNodeFilter Filter,ENodePlace ENodePlace )
-{
-	switch( Filter )
-	{
-		case F_OUTSIDE:
-		case F_COPLANAR_OUTSIDE:
-		case F_COSPATIAL_FACING_IN:
-		case F_COSPATIAL_FACING_OUT:
-			// Ignore.
-			break;
-		case F_INSIDE:
-		case F_COPLANAR_INSIDE:
-			if( EdPoly->Fix()>=3 )
-				new(GModel->Polys->Element)FPoly(*EdPoly);
-			break;
-	}
-}
-
-void IntersectWorldWithBrushFunc( UModel *Model, int32 iNode, FPoly *EdPoly,
-	EPolyNodeFilter Filter,ENodePlace ENodePlace )
-{
-	switch( Filter )
-	{
-		case F_OUTSIDE:
-		case F_COPLANAR_OUTSIDE:
-		case F_COSPATIAL_FACING_IN:
-			// Ignore.
-			break;
-		case F_INSIDE:
-		case F_COPLANAR_INSIDE:
-		case F_COSPATIAL_FACING_OUT:
-			if( EdPoly->Fix() >= 3 )
-				new(GModel->Polys->Element)FPoly(*EdPoly);
-			break;
-	}
-}
-
-void DeIntersectBrushWithWorldFunc( UModel* Model, int32 iNode, FPoly* EdPoly,
-	EPolyNodeFilter Filter,ENodePlace ENodePlace )
-{
-	switch( Filter )
-	{
-		case F_INSIDE:
-		case F_COPLANAR_INSIDE:
-		case F_COSPATIAL_FACING_OUT:
-		case F_COSPATIAL_FACING_IN:
-			// Ignore.
-			break;
-		case F_OUTSIDE:
-		case F_COPLANAR_OUTSIDE:
-			if( EdPoly->Fix()>=3 )
-				new(GModel->Polys->Element)FPoly(*EdPoly);
-			break;
-	}
-}
-
-void DeIntersectWorldWithBrushFunc( UModel* Model, int32 iNode, FPoly* EdPoly,
-	EPolyNodeFilter Filter,ENodePlace ENodePlace )
-{
-	switch( Filter )
-	{
-		case F_OUTSIDE:
-		case F_COPLANAR_OUTSIDE:
-		case F_COSPATIAL_FACING_OUT:
-			// Ignore.
-			break;
-		case F_COPLANAR_INSIDE:
-		case F_INSIDE:
-		case F_COSPATIAL_FACING_IN:
-			if( EdPoly->Fix() >= 3 )
-	        {
-				EdPoly->Reverse();
-				new(GModel->Polys->Element)FPoly(*EdPoly);
-				EdPoly->Reverse();
-			}
-			break;
-	}
-}
-
-/*----------------------------------------------------------------------------
-   CSG polygon filtering routine (calls the callbacks).
-----------------------------------------------------------------------------*/
-
-//
-// Handle a piece of a polygon that was filtered to a leaf.
-//
-void FilterLeaf
-(
-	BSP_FILTER_FUNC FilterFunc, 
-	UModel*			Model,
-	int32			    iNode, 
-	FPoly*			EdPoly, 
-	FCoplanarInfo	CoplanarInfo, 
-	int32				LeafOutside, 
-	ENodePlace		ENodePlace
-)
-{
-	EPolyNodeFilter FilterType;
-
-	if( CoplanarInfo.iOriginalNode == INDEX_NONE )
-	{
-		// Processing regular, non-coplanar polygons.
-		FilterType = LeafOutside ? F_OUTSIDE : F_INSIDE;
-		FilterFunc( Model, iNode, EdPoly, FilterType, ENodePlace );
-	}
-	else if( CoplanarInfo.ProcessingBack )
-	{
-		// Finished filtering polygon through tree in back of parent coplanar.
-		DoneFilteringBack:
-		if      ((!LeafOutside) && (!CoplanarInfo.FrontLeafOutside)) FilterType = F_COPLANAR_INSIDE;
-		else if (( LeafOutside) && ( CoplanarInfo.FrontLeafOutside)) FilterType = F_COPLANAR_OUTSIDE;
-		else if ((!LeafOutside) && ( CoplanarInfo.FrontLeafOutside)) FilterType = F_COSPATIAL_FACING_OUT;
-		else if (( LeafOutside) && (!CoplanarInfo.FrontLeafOutside)) FilterType = F_COSPATIAL_FACING_IN;
-		else
-		{
-			UE_LOG(LogEditorBsp, Fatal, TEXT("FilterLeaf: Bad Locs"));
-			return;
-		}
-		FilterFunc( Model, CoplanarInfo.iOriginalNode, EdPoly, FilterType, FBSPOps::NODE_Plane );
-	}
-	else
-	{
-		CoplanarInfo.FrontLeafOutside = LeafOutside;
-
-		if( CoplanarInfo.iBackNode == INDEX_NONE )
-		{
-			// Back tree is empty.
-			LeafOutside = CoplanarInfo.BackNodeOutside;
-			goto DoneFilteringBack;
-		}
-		else
-		{
-			// Call FilterEdPoly to filter through the back.  This will result in
-			// another call to FilterLeaf with iNode = leaf this falls into in the
-			// back tree and EdPoly = the final EdPoly to insert.
-			CoplanarInfo.ProcessingBack=1;
-			FilterEdPoly( FilterFunc, Model, CoplanarInfo.iBackNode, EdPoly,CoplanarInfo, CoplanarInfo.BackNodeOutside );
-		}
-	}
-}
-
-//
-// Filter an EdPoly through the Bsp recursively, calling FilterFunc
-// for all chunks that fall into leaves.  FCoplanarInfo is used to
-// handle the tricky case of double-recursion for polys that must be
-// filtered through a node's front, then filtered through the node's back,
-// in order to handle coplanar CSG properly.
-//
-void FilterEdPoly
-(
-	BSP_FILTER_FUNC	FilterFunc, 
-	UModel			*Model,
-	int32			    iNode, 
-	FPoly			*EdPoly, 
-	FCoplanarInfo	CoplanarInfo, 
-	int32				Outside
-)
-{
-	int32            SplitResult,iOurFront,iOurBack;
-	int32			   NewFrontOutside,NewBackOutside;
-
-	FilterLoop:
-
-	// Split em.
-	FPoly TempFrontEdPoly,TempBackEdPoly;
-	SplitResult = EdPoly->SplitWithPlane
-	(
-		Model->Points [Model->Verts[Model->Nodes[iNode].iVertPool].pVertex],
-		Model->Vectors[Model->Surfs[Model->Nodes[iNode].iSurf].vNormal],
-		&TempFrontEdPoly,
-		&TempBackEdPoly,
-		0
-	);
-
-	// Process split results.
-	if( SplitResult == SP_Front )
-	{
-		Front:
-
-		FBspNode *Node = &Model->Nodes[iNode];
-		Outside        = Outside || Node->IsCsg();
-
-		if( Node->iFront == INDEX_NONE )
-		{
-			FilterLeaf(FilterFunc,Model,iNode,EdPoly,CoplanarInfo,Outside,FBSPOps::NODE_Front);
-		}
-		else
-		{
-			iNode = Node->iFront;
-			goto FilterLoop;
-		}
-	}
-	else if( SplitResult == SP_Back )
-	{
-		FBspNode *Node = &Model->Nodes[iNode];
-		Outside        = Outside && !Node->IsCsg();
-
-		if( Node->iBack == INDEX_NONE )
-		{
-			FilterLeaf( FilterFunc, Model, iNode, EdPoly, CoplanarInfo, Outside, FBSPOps::NODE_Back );
-		}
-		else
-		{
-			iNode=Node->iBack;
-			goto FilterLoop;
-		}
-	}
-	else if( SplitResult == SP_Coplanar )
-	{
-		if( CoplanarInfo.iOriginalNode != INDEX_NONE )
-		{
-			// This will happen once in a blue moon when a polygon is barely outside the
-			// coplanar threshold and is split up into a new polygon that is
-			// is barely inside the coplanar threshold.  To handle this, just classify
-			// it as front and it will be handled propery.
-			FBSPOps::GErrors++;
-// 			UE_LOG(LogEditorBsp, Warning, TEXT("FilterEdPoly: Encountered out-of-place coplanar") );
-			goto Front;
-		}
-		CoplanarInfo.iOriginalNode        = iNode;
-		CoplanarInfo.iBackNode            = INDEX_NONE;
-		CoplanarInfo.ProcessingBack       = 0;
-		CoplanarInfo.BackNodeOutside      = Outside;
-		NewFrontOutside                   = Outside;
-
-		// See whether Node's iFront or iBack points to the side of the tree on the front
-		// of this polygon (will be as expected if this polygon is facing the same
-		// way as first coplanar in link, otherwise opposite).
-		if( (FVector(Model->Nodes[iNode].Plane) | EdPoly->Normal) >= 0.0 )
-		{
-			iOurFront = Model->Nodes[iNode].iFront;
-			iOurBack  = Model->Nodes[iNode].iBack;
-			
-			if( Model->Nodes[iNode].IsCsg() )
-	        {
-				CoplanarInfo.BackNodeOutside = 0;
-				NewFrontOutside              = 1;
-			}
-		}
-		else
-		{
-			iOurFront = Model->Nodes[iNode].iBack;
-			iOurBack  = Model->Nodes[iNode].iFront;
-			
-			if( Model->Nodes[iNode].IsCsg() )
-	        {
-				CoplanarInfo.BackNodeOutside = 1; 
-				NewFrontOutside              = 0;
-			}
-		}
-
-		// Process front and back.
-		if ((iOurFront==INDEX_NONE)&&(iOurBack==INDEX_NONE))
-		{
-			// No front or back.
-			CoplanarInfo.ProcessingBack		= 1;
-			CoplanarInfo.FrontLeafOutside	= NewFrontOutside;
-			FilterLeaf
-			(
-				FilterFunc,
-				Model,
-				iNode,
-				EdPoly,
-				CoplanarInfo,
-				CoplanarInfo.BackNodeOutside,
-				FBSPOps::NODE_Plane
-			);
-		}
-		else if( iOurFront==INDEX_NONE && iOurBack!=INDEX_NONE )
-		{
-			// Back but no front.
-			CoplanarInfo.ProcessingBack		= 1;
-			CoplanarInfo.iBackNode			= iOurBack;
-			CoplanarInfo.FrontLeafOutside	= NewFrontOutside;
-
-			iNode   = iOurBack;
-			Outside = CoplanarInfo.BackNodeOutside;
-			goto FilterLoop;
-		}
-		else
-		{
-			// Has a front and maybe a back.
-
-			// Set iOurBack up to process back on next call to FilterLeaf, and loop
-			// to process front.  Next call to FilterLeaf will set FrontLeafOutside.
-			CoplanarInfo.ProcessingBack = 0;
-
-			// May be a node or may be INDEX_NONE.
-			CoplanarInfo.iBackNode = iOurBack;
-
-			iNode   = iOurFront;
-			Outside = NewFrontOutside;
-			goto FilterLoop;
-		}
-	}
-	else if( SplitResult == SP_Split )
-	{
-		// Front half of split.
-		if( Model->Nodes[iNode].IsCsg() )
-		{
-			NewFrontOutside = 1; 
-			NewBackOutside  = 0;
-		}
-		else
-		{
-			NewFrontOutside = Outside;
-			NewBackOutside  = Outside;
-		}
-
-		if( Model->Nodes[iNode].iFront==INDEX_NONE )
-		{
-			FilterLeaf
-			(
-				FilterFunc,
-				Model,
-				iNode,
-				&TempFrontEdPoly,
-				CoplanarInfo,
-				NewFrontOutside,
-				FBSPOps::NODE_Front
-			);
-		}
-		else
-		{
-			FilterEdPoly
-			(
-				FilterFunc,
-				Model,
-				Model->Nodes[iNode].iFront,
-				&TempFrontEdPoly,
-				CoplanarInfo,
-				NewFrontOutside
-			);
-		}
-
-		// Back half of split.
-		if( Model->Nodes[iNode].iBack==INDEX_NONE )
-		{
-			FilterLeaf
-			(
-				FilterFunc,
-				Model,
-				iNode,
-				&TempBackEdPoly,
-				CoplanarInfo,
-				NewBackOutside,
-				FBSPOps::NODE_Back
-			);
-		}
-		else
-		{
-			FilterEdPoly
-			(
-				FilterFunc,
-				Model,
-				Model->Nodes[iNode].iBack,
-				&TempBackEdPoly,
-				CoplanarInfo,
-				NewBackOutside
-			);
-		}
-	}
-}
-
-//
-// Regular entry into FilterEdPoly (so higher-level callers don't have to
-// deal with unnecessary info). Filters starting at root.
-//
-void BspFilterFPoly( BSP_FILTER_FUNC FilterFunc, UModel *Model, FPoly *EdPoly )
-{
-	FCoplanarInfo StartingCoplanarInfo;
-	StartingCoplanarInfo.iOriginalNode = INDEX_NONE;
-	if( Model->Nodes.Num() == 0 )
-	{
-		// If Bsp is empty, process at root.
-		FilterFunc( Model, 0, EdPoly, Model->RootOutside ? F_OUTSIDE : F_INSIDE, FBSPOps::NODE_Root );
-	}
-	else
-	{
-		// Filter through Bsp.
-		FilterEdPoly( FilterFunc, Model, 0, EdPoly, StartingCoplanarInfo, Model->RootOutside );
-	}
+	FBSPUtils::bspCleanup(Model);
 }
 
 
@@ -918,166 +39,7 @@ int UEditorEngine::bspNodeToFPoly
 	FPoly	*EdPoly
 )
 {
-	FPoly MasterEdPoly;
-
-	FBspNode &Node     	= Model->Nodes[iNode];
-	FBspSurf &Poly     	= Model->Surfs[Node.iSurf];
-	FVert	 *VertPool	= &Model->Verts[ Node.iVertPool ];
-
-	EdPoly->Base		= Model->Points [Poly.pBase];
-	EdPoly->Normal      = Model->Vectors[Poly.vNormal];
-
-	EdPoly->PolyFlags 	= Poly.PolyFlags & ~(PF_EdCut | PF_EdProcessed | PF_Selected | PF_Memorized);
-	EdPoly->iLinkSurf   = Node.iSurf;
-	EdPoly->Material    = Poly.Material;
-
-	EdPoly->Actor    	= Poly.Actor;
-	EdPoly->iBrushPoly  = Poly.iBrushPoly;
-	
-	if( polyFindMaster(Model,Node.iSurf,MasterEdPoly) )
-		EdPoly->ItemName  = MasterEdPoly.ItemName;
-	else
-		EdPoly->ItemName  = NAME_None;
-
-	EdPoly->TextureU = Model->Vectors[Poly.vTextureU];
-	EdPoly->TextureV = Model->Vectors[Poly.vTextureV];
-
-	EdPoly->LightMapScale = Poly.LightMapScale;
-
- 	EdPoly->LightmassSettings = Model->LightmassSettings[Poly.iLightmassIndex];
-
-	EdPoly->Vertices.Empty();
-
-	for(int32 VertexIndex = 0;VertexIndex < Node.NumVertices;VertexIndex++)
-	{
-		new(EdPoly->Vertices) FVector3f(Model->Points[VertPool[VertexIndex].pVertex]);
-	}
-
-	if(EdPoly->Vertices.Num() < 3)
-	{
-		EdPoly->Vertices.Empty();
-	}
-	else
-	{
-		// Remove colinear points and identical points (which will appear
-		// if T-joints were eliminated).
-		EdPoly->RemoveColinears();
-	}
-
-	return EdPoly->Vertices.Num();
-}
-
-/*---------------------------------------------------------------------------------------
-   World filtering.
----------------------------------------------------------------------------------------*/
-
-//
-// Filter all relevant world polys through the brush.
-//
-void FilterWorldThroughBrush 
-(
-	UModel*		Model,
-	UModel*		Brush,
-	EBrushType	BrushType,
-	ECsgOper	CSGOper,
-	int32		    iNode,
-	FSphere*	BrushSphere
-)
-{
-	// Loop through all coplanars.
-	while( iNode != INDEX_NONE )
-	{
-		// Get surface.
-		int32 iSurf = Model->Nodes[iNode].iSurf;
-
-		// Skip new nodes and their children, which are guaranteed new.
-		if( Model->Nodes[iNode].NodeFlags & NF_IsNew )
-			return;
-
-		// Sphere reject.
-		int DoFront = 1, DoBack = 1;
-		if( BrushSphere )
-		{
-			float Dist = Model->Nodes[iNode].Plane.PlaneDot( BrushSphere->Center );
-			DoFront    = (Dist >= -BrushSphere->W);
-			DoBack     = (Dist <= +BrushSphere->W);
-		}
-
-		// Process only polys that aren't empty.
-		FPoly TempEdPoly;
-		if( DoFront && DoBack && (GEditor->bspNodeToFPoly(Model,iNode,&TempEdPoly)>0) )
-		{
-			TempEdPoly.Actor      = Model->Surfs[iSurf].Actor;
-			TempEdPoly.iBrushPoly = Model->Surfs[iSurf].iBrushPoly;
-
-			if( BrushType==Brush_Add || BrushType==Brush_Subtract )
-			{
-				// Add and subtract work the same in this step.
-				GNode       = iNode;
-				GModel  	= Model;
-				GDiscarded  = 0;
-				GNumNodes	= Model->Nodes.Num();
-
-				// Find last coplanar in chain.
-				GLastCoplanar = iNode;
-				while( Model->Nodes[GLastCoplanar].iPlane != INDEX_NONE )
-					GLastCoplanar = Model->Nodes[GLastCoplanar].iPlane;
-
-				// Do the filter operation.
-				BspFilterFPoly
-				(
-					BrushType==Brush_Add ? AddWorldToBrushFunc : SubtractWorldToBrushFunc,
-					Brush,
-					&TempEdPoly
-				);
-				
-				if( GDiscarded == 0 )
-				{
-					// Get rid of all the fragments we added.
-					Model->Nodes[GLastCoplanar].iPlane = INDEX_NONE;
-					const bool bAllowShrinking = false;
-					Model->Nodes.RemoveAt( GNumNodes, Model->Nodes.Num()-GNumNodes, bAllowShrinking );
-				}
-				else
-				{
-					// Tag original world poly for deletion; has been deleted or replaced by partial fragments.
-					if( GModel->Nodes[GNode].NumVertices )
-					{
-						GModel->Nodes[GNode].NumVertices = 0;
-					}
-				}
-			}
-			else if( CSGOper == CSG_Intersect )
-			{
-				BspFilterFPoly( IntersectWorldWithBrushFunc, Brush, &TempEdPoly );
-			}
-			else if( CSGOper == CSG_Deintersect )
-			{
-				BspFilterFPoly( DeIntersectWorldWithBrushFunc, Brush, &TempEdPoly );
-			}
-		}
-
-		// Now recurse to filter all of the world's children nodes.
-		if( DoFront && (Model->Nodes[iNode].iFront != INDEX_NONE)) FilterWorldThroughBrush
-		(
-			Model,
-			Brush,
-			BrushType,
-			CSGOper,
-			Model->Nodes[iNode].iFront,
-			BrushSphere
-		);
-		if( DoBack && (Model->Nodes[iNode].iBack != INDEX_NONE) ) FilterWorldThroughBrush
-		(
-			Model,
-			Brush,
-			BrushType,
-			CSGOper,
-			Model->Nodes[iNode].iBack,
-			BrushSphere
-		);
-		iNode = Model->Nodes[iNode].iPlane;
-	}
+	return FBSPUtils::bspNodeToFPoly(Model, iNode, EdPoly);
 }
 
 
@@ -1094,812 +56,15 @@ int UEditorEngine::bspBrushCSG
 	bool		bShowProgressBar/*=true*/
 )
 {
-	uint32 NotPolyFlags = 0;
-	int32 NumPolysFromBrush=0,i,j,ReallyBig;
-	UModel* Brush = Actor->Brush;
-
-	// Note no errors.
-	FBSPOps::GErrors = 0;
-
-	// Make sure we're in an acceptable state.
-	if( !Brush )
-	{
-		return 0;
-	}
-
-	// Non-solid and semisolid stuff can only be added.
-	if( BrushType != Brush_Add )
-	{
-		NotPolyFlags |= (PF_Semisolid | PF_NotSolid);
-	}
-
-	TempModel->EmptyModel(1,1);
-
-	// Update status.
-	ReallyBig = (Brush->Polys->Element.Num() > 200) && bShowProgressBar;
-	if( ReallyBig )
-	{
-		FText Description = NSLOCTEXT("UnrealEd", "PerformingCSGOperation", "Performing CSG operation");
-		
-		if (BrushType != Brush_MAX)
-		{	
-			if (BrushType == Brush_Add)
-			{
-				Description = NSLOCTEXT("UnrealEd", "AddingBrushToWorld", "Adding brush to world");
-			}
-			else if (BrushType == Brush_Subtract)
-			{
-				Description = NSLOCTEXT("UnrealEd", "SubtractingBrushFromWorld", "Subtracting brush from world");
-			}
-		}
-		else if (CSGOper != CSG_None)
-		{
-			if (CSGOper == CSG_Intersect)
-			{
-				Description = NSLOCTEXT("UnrealEd", "IntersectingBrushWithWorld", "Intersecting brush with world");
-			}
-			else if (CSGOper == CSG_Deintersect)
-			{
-				Description = NSLOCTEXT("UnrealEd", "DeintersectingBrushWithWorld", "Deintersecting brush with world");
-			}
-		}
-
-		GWarn->BeginSlowTask( Description, true );
-		// Transform original brush poly into same coordinate system as world
-		// so Bsp filtering operations make sense.
-		GWarn->StatusUpdate(0, 0, NSLOCTEXT("UnrealEd", "Transforming", "Transforming"));
-	}
-
-
-
 	UMaterialInterface* SelectedMaterialInstance = GetSelectedObjects()->GetTop<UMaterialInterface>();
-
-	const FVector Scale = Actor->GetActorScale();
-	const FRotator Rotation = Actor->GetActorRotation();
-	const FVector Location = Actor->GetActorLocation();
-
-	const bool bIsMirrored = (Scale.X * Scale.Y * Scale.Z < 0.0f);
-
-	// Cache actor transform which is used for the geometry being built
-	Brush->OwnerLocationWhenLastBuilt = Location;
-	Brush->OwnerRotationWhenLastBuilt = Rotation;
-	Brush->OwnerScaleWhenLastBuilt = Scale;
-	Brush->bCachedOwnerTransformValid = true;
-
-	for( i=0; i<Brush->Polys->Element.Num(); i++ )
-	{
-		FPoly& CurrentPoly = Brush->Polys->Element[i];
-
-		// Set texture the first time.
-		if ( bReplaceNULLMaterialRefs )
-		{
-			UMaterialInterface*& PolyMat = CurrentPoly.Material;
-			if ( !PolyMat || PolyMat == UMaterial::GetDefaultMaterial(MD_Surface) )
-			{
-				PolyMat = SelectedMaterialInstance;
-			}
-		}
-
-		// Get the brush poly.
-		FPoly DestEdPoly = CurrentPoly;
-		check(CurrentPoly.iLink<Brush->Polys->Element.Num());
-
-		// Set its backward brush link.
-		DestEdPoly.Actor       = Actor;
-		DestEdPoly.iBrushPoly  = i;
-
-		// Update its flags.
-		DestEdPoly.PolyFlags = (DestEdPoly.PolyFlags | PolyFlags) & ~NotPolyFlags;
-
-		// Set its internal link.
-		if (DestEdPoly.iLink == INDEX_NONE)
-		{
-			DestEdPoly.iLink = i;
-		}
-
-		// Transform it.
-		DestEdPoly.Scale( Scale );
-		DestEdPoly.Rotate( Rotation );
-		DestEdPoly.Transform( Location );
-
-		// Reverse winding and normal if the parent brush is mirrored
-		if (bIsMirrored)
-		{
-			DestEdPoly.Reverse();
-			DestEdPoly.CalcNormal();
-		}
-
-		// Add poly to the temp model.
-		new(TempModel->Polys->Element)FPoly( DestEdPoly );
-	}
-	if( ReallyBig ) GWarn->StatusUpdate( 0, 0, NSLOCTEXT("UnrealEd", "FilteringBrush", "Filtering brush") );
-
-	// Pass the brush polys through the world Bsp.
-	if( CSGOper==CSG_Intersect || CSGOper==CSG_Deintersect )
-	{
-		// Empty the brush.
-		Brush->EmptyModel(1,1);
-
-		// Intersect and deintersect.
-		for( i=0; i<TempModel->Polys->Element.Num(); i++ )
-		{
-         	FPoly EdPoly = TempModel->Polys->Element[i];
-			GModel = Brush;
-			// TODO: iLink / iLinkSurf in EdPoly / TempModel->Polys->Element[i] ?
-			BspFilterFPoly( CSGOper==CSG_Intersect ? IntersectBrushWithWorldFunc : DeIntersectBrushWithWorldFunc, Model,  &EdPoly );
-		}
-		NumPolysFromBrush = Brush->Polys->Element.Num();
-	}
-	else
-	{
-		// Add and subtract.
-		TMap<int32, int32> SurfaceIndexRemap;
-		for( i=0; i<TempModel->Polys->Element.Num(); i++ )
-		{
-         	FPoly EdPoly = TempModel->Polys->Element[i];
-
-         	// Mark the polygon as non-cut so that it won't be harmed unless it must
-         	// be split, and set iLink so that BspAddNode will know to add its information
-         	// if a node is added based on this poly.
-         	EdPoly.PolyFlags &= ~(PF_EdCut);
-			const int32* SurfaceIndexPtr = SurfaceIndexRemap.Find(EdPoly.iLink);
-			if (SurfaceIndexPtr == nullptr)
-			{
-				const int32 NewSurfaceIndex = Model->Surfs.Num();
-				SurfaceIndexRemap.Add(EdPoly.iLink, NewSurfaceIndex);
-				EdPoly.iLinkSurf = TempModel->Polys->Element[i].iLinkSurf = NewSurfaceIndex;
-			}
-			else
-			{
-				EdPoly.iLinkSurf = TempModel->Polys->Element[i].iLinkSurf = *SurfaceIndexPtr;
-			}
-
-			// Filter brush through the world.
-			BspFilterFPoly( BrushType==Brush_Add ? AddBrushToWorldFunc : SubtractBrushFromWorldFunc, Model, &EdPoly );
-		}
-	}
-	if( Model->Nodes.Num() && !(PolyFlags & (PF_NotSolid | PF_Semisolid)) )
-	{
-		// Quickly build a Bsp for the brush, tending to minimize splits rather than balance
-		// the tree.  We only need the cutting planes, though the entire Bsp struct (polys and
-		// all) is built.
-
-		FBspPointsGrid* LevelModelPointsGrid = FBspPointsGrid::GBspPoints;
-		FBspPointsGrid* LevelModelVectorsGrid = FBspPointsGrid::GBspVectors;
-
-		// For the bspBuild call, temporarily create a new pair of BspPointsGrids for the TempModel.
-		TUniquePtr<FBspPointsGrid> BspPoints = MakeUnique<FBspPointsGrid>(50.0f, THRESH_POINTS_ARE_SAME);
- 		TUniquePtr<FBspPointsGrid> BspVectors = MakeUnique<FBspPointsGrid>(1 / 16.0f, FMath::Max(THRESH_NORMALS_ARE_SAME, THRESH_VECTORS_ARE_NEAR));
- 		FBspPointsGrid::GBspPoints = BspPoints.Get();
- 		FBspPointsGrid::GBspVectors = BspVectors.Get();
-
-		if( ReallyBig ) GWarn->StatusUpdate( 0, 0, NSLOCTEXT("UnrealEd", "BuildingBSP", "Building BSP") );
-		
-		FBSPOps::bspBuild( TempModel, FBSPOps::BSP_Lame, 0, 70, 1, 0 );
-
-		// Reinstate the original BspPointsGrids used for building the level Model.
-		FBspPointsGrid::GBspPoints = LevelModelPointsGrid;
-		FBspPointsGrid::GBspVectors = LevelModelVectorsGrid;
-
-		if( ReallyBig ) GWarn->StatusUpdate( 0, 0, NSLOCTEXT("UnrealEd", "FilteringWorld", "Filtering world") );
-		GModel = Brush;
-		TempModel->BuildBound();
-
-		FSphere	BrushSphere = TempModel->Bounds.GetSphere();
-		FilterWorldThroughBrush( Model, TempModel, BrushType, CSGOper, 0, &BrushSphere );
-	}
-	if( CSGOper==CSG_Intersect || CSGOper==CSG_Deintersect )
-	{
-		if( ReallyBig ) GWarn->StatusUpdate( 0, 0, NSLOCTEXT("UnrealEd", "AdjustingBrush", "Adjusting brush") );
-
-		// Link polys obtained from the original brush.
-		for( i=NumPolysFromBrush-1; i>=0; i-- )
-		{
-			FPoly *DestEdPoly = &Brush->Polys->Element[i];
-			for( j=0; j<i; j++ )
-			{
-				if( DestEdPoly->iLink == Brush->Polys->Element[j].iLink )
-				{
-					DestEdPoly->iLink = j;
-					break;
-				}
-			}
-			if( j >= i ) DestEdPoly->iLink = i;
-		}
-
-		// Link polys obtained from the world.
-		for( i=Brush->Polys->Element.Num()-1; i>=NumPolysFromBrush; i-- )
-		{
-			FPoly *DestEdPoly = &Brush->Polys->Element[i];
-			for( j=NumPolysFromBrush; j<i; j++ )
-			{
-				if( DestEdPoly->iLink == Brush->Polys->Element[j].iLink )
-				{
-					DestEdPoly->iLink = j;
-					break;
-				}
-			}
-			if( j >= i ) DestEdPoly->iLink = i;
-		}
-		Brush->Linked = 1;
-
-		// Detransform the obtained brush back into its original coordinate system.
-		for( i=0; i<Brush->Polys->Element.Num(); i++ )
-		{
-			FPoly *DestEdPoly = &Brush->Polys->Element[i];
-			DestEdPoly->Transform(-Location);
-			DestEdPoly->Rotate(Rotation.GetInverse());
-			DestEdPoly->Scale(FVector(1.0f) / Scale);
-			DestEdPoly->Fix();
-			DestEdPoly->Actor		= NULL;
-			DestEdPoly->iBrushPoly	= i;
-		}
-	}
-
-	if( BrushType==Brush_Add || BrushType==Brush_Subtract )
-	{
-		// Clean up nodes, reset node flags.
-		bspCleanup( Model );
-
-		// Rebuild bounding volumes.
-		if( bBuildBounds )
-		{
-			FBSPOps::bspBuildBounds( Model );
-		}
-	}
-
-	Brush->NumUniqueVertices = TempModel->Points.Num();
-	// Release TempModel.
-	TempModel->EmptyModel(1,1);
-	
-	// Merge coplanars if needed.
-	if( CSGOper==CSG_Intersect || CSGOper==CSG_Deintersect )
-	{
-		if( ReallyBig )
-		{
-			GWarn->StatusUpdate( 0, 0, NSLOCTEXT("UnrealEd", "Merging", "Merging") );
-		}
-		if( bMergePolys )
-		{
-			bspMergeCoplanars( Brush, 1, 0 );
-		}
-	}
-	if( ReallyBig )
-	{
-		GWarn->EndSlowTask();
-	}
-
-	return 1 + FBSPOps::GErrors;
-}
-
-/*---------------------------------------------------------------------------------------
-   Functions for maintaining linked geometry lists.
----------------------------------------------------------------------------------------*/
-
-//
-// A node and vertex number corresponding to a point, used in generating Bsp side links.
-//
-class FPointVert 
-{
-public:
-	int32		    iNode;
-	int32		    nVertex;
-	FPointVert* Next;
-};
-
-//
-// A list of point/vertex links, used in generating Bsp side links.
-//
-class FPointVertList
-{
-public:
-
-	// Variables.
-	UModel		*Model;
-	FPointVert	**Index;
-	FMemMark*	Mark;
-
-	/** Initialization constructor. */
-	FPointVertList()
-	:	Model(NULL)
-	,	Index(NULL)
-	,	Mark(NULL)
-	{}
-
-	/** Destructor. */
-	~FPointVertList()
-	{
-		check(!Mark);
-	}
-
-	// Allocate.
-	void Alloc( UModel *ThisModel )
-	{
-		check(!Mark);
-		Mark = new FMemMark(FMemStack::Get());
-
-		// Allocate working tables.
-		Model = ThisModel;
-		Index = new(FMemStack::Get(),MEM_Zeroed,Model->Points.Num())FPointVert*;
-
-	}
-
-	// Free.
-	void Free()
-	{
-		delete Mark;
-		Mark = NULL;
-	}
-
-	// Add all of a node's vertices to a node-vertex list.
-	void AddNode( int32 iNode )
-	{
-		FBspNode&	Node 	 =  Model->Nodes[iNode];
-		FVert*		VertPool =  &Model->Verts[Node.iVertPool];
-
-		for( int32 i=0; i < Node.NumVertices; i++ )
-		{
-			int32 pVertex = VertPool[i].pVertex;
-
-			// Add new point/vertex pair to array, and insert new array entry
-			// between index and first entry.
-			FPointVert *PointVert	= new(FMemStack::Get())FPointVert;
-			
-			PointVert->iNode 		= iNode;
-			PointVert->nVertex		= i;
-			PointVert->Next			= Index[pVertex];
-
-			Index[pVertex]			= PointVert;
-		}
-	}
-
-	// Add all nodes' vertices in the model to a node-vertex list.
-	void AddAllNodes()
-	{
-		for( int32 iNode=0; iNode < Model->Nodes.Num(); iNode++ )
-			AddNode( iNode );
-
-	}
-
-	// Remove all of a node's vertices from a node-vertex list.
-	void RemoveNode( int32 iNode )
-	{
-		FBspNode&	Node 		= Model->Nodes[iNode];
-		FVert*		VertPool	= &Model->Verts[Node.iVertPool];
-
-		// Loop through all of the node's vertices and search through the
-		// corresponding point's node-vert list, and delink this node.
-		int Count=0;
-		for( int32 i = 0; i < Node.NumVertices; i++ )
-		{
-			int32 pVertex = VertPool[i].pVertex;
-
-			for( FPointVert **PrevLink = &Index[pVertex]; *PrevLink; PrevLink=&(*PrevLink)->Next )
-			{
-				if( (*PrevLink)->iNode == iNode )
-				{
-					// Delink this entry from the list.
-					*PrevLink = (*PrevLink)->Next;
-					Count++;
-
-					if( *PrevLink == NULL )
-						break;
-				}
-			}
-
-			// Node's vertex wasn't found, there's a bug.
-			check(Count>=1);
-		}
-	}
-};
-
-/*---------------------------------------------------------------------------------------
-   Geometry optimization.
----------------------------------------------------------------------------------------*/
-
-//
-// Add a point to a Bsp node before a specified vertex (between it and the previous one).
-// VertexNumber can be from 0 (before first) to Node->NumVertices (after last).
-//
-// Splits node into two coplanar polys if necessary. If the polygon is split, the
-// vertices will be distributed among this node and it's newly-linked iPlane node
-// in an arbitrary way, that preserves the clockwise orientation of the vertices.
-//
-// Maintains node-vertex list, if not NULL.
-//
-void AddPointToNode
-(
-	UModel*         Model,
-	FPointVertList* PointVerts,
-	int32			    iNode, 
-	int32			    VertexNumber, 
-	int32			    pVertex 
-)
-{
-	FBspNode&	Node = Model->Nodes[iNode];
-
-	if( Node.NumVertices >= FBspNode::MAX_NODE_VERTICES - 1 )
-	{
-		// Just refuse to add point: This is a non-fatal problem.
-//		UE_LOG(LogEditorBsp, Warning, TEXT("Node side limit reached") );
-		return;
-	}
-
-	// Remove node from vertex list, since vertex numbers will be reordered.
-	if( PointVerts )
-		PointVerts->RemoveNode( iNode );
-
-	int32 iOldVert = Node.iVertPool;
-
-	Node.iVertPool = Model->Verts.AddUninitialized( Node.NumVertices+1 );
-
-	// Make sure this node doesn't already contain the vertex.
-	for( int32 i=0; i<Node.NumVertices; i++ )
-		check( Model->Verts[iOldVert + i].pVertex != pVertex );
-
-	// Copy the old vertex pool to the new one.
-	for( int32 i=0; i<VertexNumber; i++ )
-		Model->Verts[Node.iVertPool + i] = Model->Verts[iOldVert + i];
-
-	for( int32 i=VertexNumber; i<Node.NumVertices; i++ )
-		Model->Verts[Node.iVertPool + i + 1] = Model->Verts[iOldVert + i];
-
-	// Add the new point to the new vertex pool.
-	Model->Verts[Node.iVertPool + VertexNumber].pVertex = pVertex;
-	Model->Verts[Node.iVertPool + VertexNumber].iSide   = INDEX_NONE;
-
-	// Increment number of node vertices.
-	Node.NumVertices++;
-
-	// Update the point-vertex list.
-	if( PointVerts )
-		PointVerts->AddNode( iNode );
-
-}
-
-//
-// Add a point to all sides of polygons in which the side intersects with
-// this point but doesn't contain it, and has the correct (clockwise) orientation
-// as this side.  pVertex is the index of the point to handle, and
-// ReferenceVertex defines the direction of this side.
-//
-int DistributePoint
-(
-	UModel*         Model,
-	FPointVertList* PointVerts,	
-	int32			    iNode,
-	int32			    pVertex
-)
-{
-	int32 Count = 0;
-
-	// Handle front, back, and plane.
-	float Dist = Model->Nodes[iNode].Plane.PlaneDot(Model->Points[pVertex]);
-	if( Dist < THRESH_OPTGEOM_COPLANAR )
-	{
-		// Back.
-		if( Model->Nodes[iNode].iBack != INDEX_NONE )
-			Count += DistributePoint( Model, PointVerts, Model->Nodes[iNode].iBack, pVertex );
-	}
-	if( Dist > -THRESH_OPTGEOM_COPLANAR )
-	{
-		// Front.
-		if( Model->Nodes[iNode].iFront != INDEX_NONE )
-			Count += DistributePoint( Model, PointVerts, Model->Nodes[iNode].iFront, pVertex );
-	}
-	if( Dist > -THRESH_OPTGEOM_COPLANAR && Dist < THRESH_OPTGEOM_COPLANAR )
-	{
-		// This point is coplanar with this node, so check point for intersection with
-		// this node's sides, then loop with its coplanars.
-		for( ; iNode!=INDEX_NONE; iNode=Model->Nodes[iNode].iPlane )
-		{
-			FVert* VertPool = &Model->Verts[Model->Nodes[iNode].iVertPool];
-
-			// Skip this node if it already contains the point in question.
-			int32 i;
-			for( i=0; i<Model->Nodes[iNode].NumVertices; i++ )
-				if( VertPool[i].pVertex == pVertex )
-					break;
-			if( i != Model->Nodes[iNode].NumVertices )
-				continue;
-
-			// Loop through all sides and see if (A) side is colinear with point, and
-			// (B) point falls within inside of this side.
-			int32 FoundSide       = -1;
-			int32 SkippedColinear = 0;
-			int32 SkippedInside   = 0;
-
-			for( i=0; i<Model->Nodes[iNode].NumVertices; i++ )
-			{
-				int32 j = (i>0) ? (i-1) : (Model->Nodes[iNode].NumVertices-1);
-
-				// Create cutting plane perpendicular to both this side and the polygon's normal.
-				FVector Side            = Model->Points[VertPool[i].pVertex] - Model->Points[VertPool[j].pVertex];
-				FVector SidePlaneNormal = Side ^ Model->Nodes[iNode].Plane;
-				float   SizeSquared     = SidePlaneNormal.SizeSquared();
-
-				if( SizeSquared>FMath::Square(0.001) )
-				{
-					// Points aren't coincedent.
-					Dist = ((Model->Points[pVertex]-Model->Points[VertPool[i].pVertex]) | SidePlaneNormal)/FMath::Sqrt(SizeSquared);
-					if( Dist >= THRESH_OPTGEOM_COSIDAL )
-					{
-						// Point is outside polygon, can't possibly fall on a side.
-						break;
-					}
-					else if( Dist > -THRESH_OPTGEOM_COSIDAL )
-					{
-						// The point we're adding falls on this line.
-						//
-						// Verify that it falls within this side; though it's colinear
-						// it may be out of the bounds of the line's endpoints if this side
-						// is colinear with an adjacent side.
-						//
-						// Do this by checking distance from point to side's midpoint and
-						// comparing with the side's half-length.
-
-						FVector MidPoint    = (Model->Points[VertPool[i].pVertex] + Model->Points[VertPool[j].pVertex])*0.5;
-						FVector MidDistVect = Model->Points[pVertex] - MidPoint;
-						if( MidDistVect.SizeSquared() <= FMath::Square(0.501)*Side.SizeSquared() )
-						{
-							FoundSide = i;
-						}
-						else
-						{
-							SkippedColinear = 1;
-						}
-					}
-					else
-					{
-						// Point is inside polygon, so continue checking.
-						SkippedInside = 1;
-					}
-				}
-				else
-				{
-					FBSPOps::GErrors++;
-					//UE_LOG(LogEditorBsp, Warning, "Found tiny side" );
-					//Poly.PolyFlags |= PF_Selected;
-				}
-			}
-			if( i==Model->Nodes[iNode].NumVertices && FoundSide>=0 )
-			{
-				// AddPointToNode will reorder the vertices in this node.  This is okay
-				// because it's called outside of the vertex loop.
-				AddPointToNode( Model, PointVerts, iNode, FoundSide, pVertex );
-				Count++;
-			}
-			else if( SkippedColinear )
-			{
-				// This happens occasionally because of the fuzzy Dist comparison.  It is
-				// not a sign of a problem when the vertex being distributed is colinear
-				// with one of this polygon's sides, but slightly outside of this polygon.
-
-				FBSPOps::GErrors++;
-				//UE_LOG(LogEditorBsp, Warning, "Skipped colinear" );
-				//Poly.PolyFlags |= PF_Selected;
-			}
-			else if( SkippedInside )
-			{
-				// Point is on interior of polygon.
-				FBSPOps::GErrors++;
-				//UE_LOG(LogEditorBsp, Warning, "Skipped interior" );
-				//Poly.PolyFlags |= PF_Selected;
-			}
-		}
-	}
-	return Count;
-}
-
-//
-// Merge points that are near.
-//
-void MergeNearPoints( UModel *Model, float Dist )
-{
-	FMemMark Mark(FMemStack::Get());
-	int32* PointRemap = new(FMemStack::Get(),Model->Points.Num())int32;
-	int32 Merged=0,Collapsed=0;
-
-	// Find nearer point for all points.
-	for( int32 i=0; i<Model->Points.Num(); i++ )
-	{
-		PointRemap[i] = i;
-		FVector3f &Point = Model->Points[i];
-		for( int32 j=0; j<i; j++ )
-		{
-			FVector3f &TestPoint = Model->Points[j];
-			if( (TestPoint - Point).SizeSquared() < Dist*Dist )
-			{
-				PointRemap[i] = j;
-				Merged++;
-				break;
-			}
-		}
-	}
-
-	// Remap VertPool.
-	for( int32 i=0; i<Model->Verts.Num(); i++ )
-	{
-		if( Model->Verts[i].pVertex>=0 && Model->Verts[i].pVertex<Model->Points.Num() )
-		{
-			CA_SUPPRESS(6001); // warning C6001: Using uninitialized memory '*PointRemap'.
-			Model->Verts[i].pVertex = PointRemap[Model->Verts[i].pVertex];
-		}
-	}
-
-	// Remap Surfs.
-	for( int32 i=0; i<Model->Surfs.Num(); i++ )
-	{
-		if( Model->Surfs[i].pBase>=0 && Model->Surfs[i].pBase<Model->Points.Num() )
-		{
-			CA_SUPPRESS(6001); // warning C6001: Using uninitialized memory '*PointRemap'.
-			Model->Surfs[i].pBase = PointRemap[Model->Surfs[i].pBase];
-		}
-	}
-
-	// Remove duplicate points from nodes.
-	for( int32 i=0; i<Model->Nodes.Num(); i++ )
-	{
-		FBspNode &Node = Model->Nodes[i];
-		FVert *Pool = &Model->Verts[Node.iVertPool];
-		
-		int k=0;
-		for( int j=0; j<Node.NumVertices; j++ )
-		{
-			FVert *A = &Pool[j];
-			FVert *B = &Pool[j ? j-1 : Node.NumVertices-1];
-			
-			if( A->pVertex != B->pVertex )
-				Pool[k++] = Pool[j];
-		}
-		Node.NumVertices = k>=3 ? k : 0;
-		if( k < 3 )
-			Collapsed++;
-	}
-
-	// Success.
-// 	UE_LOG(LogEditorBsp, Log,  TEXT("MergeNearPoints merged %i/%i, collapsed %i"), Merged, Model->Points.Num(), Collapsed );
-	Mark.Pop();
+	return FBSPUtils::bspBrushCSG(Actor, Model, TempModel, SelectedMaterialInstance, PolyFlags, BrushType, CSGOper, bBuildBounds, bMergePolys, bReplaceNULLMaterialRefs, bShowProgressBar);
 }
 
 
 void UEditorEngine::bspOptGeom( UModel *Model )
 {
-	FPointVertList PointVerts;
-
-	//UE_LOG(LogEditorBsp, Log,  TEXT("BspOptGeom begin") );
-
 	if( GUndo )
 		ResetTransaction( NSLOCTEXT("UnrealEd", "GeometryOptimization", "Geometry Optimization") );
 
-	MergeNearPoints			(Model,0.25);
-	FBSPOps::bspRefresh		(Model,0);
-	PointVerts.Alloc		(Model);
-	PointVerts.AddAllNodes	();
-
-	// First four entries are reserved for view-clipped sides.
-	Model->NumSharedSides = 4;
-	
-	// Mark all sides as unlinked.
-	int32 i;
-	for( i=0; i<Model->Verts.Num(); i++ )
-		Model->Verts[i].iSide = INDEX_NONE;
-
-	int32 TeesFound = 0, Distributed = 0;
-
-	// Eliminate T-joints on each node by finding all vertices that aren't attached to
-	// two shared sides, then filtering them down through the BSP and adding them to
-	// the sides they belong on.
-	for( int32 iNode=0; iNode < Model->Nodes.Num(); iNode++ )
-	{
-		FBspNode &Node = Model->Nodes[iNode];
-		
-		// Loop through all sides (side := line from PrevVert to ThisVert)		
-		for( uint8 ThisVert=0; ThisVert < Node.NumVertices; ThisVert++ )
-		{
-			uint8 PrevVert = (ThisVert>0) ? (ThisVert - 1) : (Node.NumVertices-1);
-
-			// Count number of nodes sharing this side, i.e. number of nodes for
-			// which two adjacent vertices are identical to this side's two vertices.
-			for( FPointVert* PV1 = PointVerts.Index[Model->Verts[Node.iVertPool+ThisVert].pVertex]; PV1; PV1=PV1->Next )
-				for( FPointVert* PV2 = PointVerts.Index[Model->Verts[Node.iVertPool+PrevVert].pVertex]; PV2; PV2=PV2->Next )
-					if( PV1->iNode==PV2->iNode && PV1->iNode!=iNode )
-						goto SkipIt;
-
-			// Didn't find another node that shares our two vertices; must add each
-			// vertex to all polygons where the vertex lies on the polygon's side.
-			// DistributePoint will not affect the current node but may change others
-			// and may increase the number of nodes in the Bsp.
-			TeesFound++;
-			Distributed = 0;
-			Distributed += DistributePoint( Model, &PointVerts, 0, Model->Verts[Node.iVertPool+ThisVert].pVertex );
-			Distributed += DistributePoint( Model, &PointVerts, 0, Model->Verts[Node.iVertPool+PrevVert].pVertex );
-			SkipIt:;
-		}
-	}
-	// Build side links
-	// Definition of side: Side (i) links node vertex (i) to vertex ((i+1)%n)
-	//UE_LOG(LogEditorBsp, Log,  TEXT("BspOptGeom building sidelinks") );
-
-	PointVerts.Free			();
-	PointVerts.Alloc		(Model);
-	PointVerts.AddAllNodes	();
-
-	for( int32 iNode=0; iNode < Model->Nodes.Num(); iNode++ )
-	{
-		FBspNode &Node = Model->Nodes[iNode];		
-		for( uint8 ThisVert=0; ThisVert < Node.NumVertices; ThisVert++ )
-		{
-			if( Model->Verts[Node.iVertPool+ThisVert].iSide == INDEX_NONE )
-			{
-				// See if this node links to another one.
-				uint8 PrevVert = (ThisVert>0) ? (ThisVert - 1) : (Node.NumVertices-1);
-				for( FPointVert* PV1=PointVerts.Index[Model->Verts[Node.iVertPool+ThisVert].pVertex]; PV1; PV1=PV1->Next )
-				{
-					for( FPointVert* PV2=PointVerts.Index[Model->Verts[Node.iVertPool+PrevVert].pVertex]; PV2; PV2=PV2->Next )
-					{
-						if( PV1->iNode==PV2->iNode && PV1->iNode!=iNode )
-						{
-							// Make sure that the other node's two vertices are adjacent and
-							// ordered opposite this node's vertices.
-							int32		    iOtherNode	= PV2->iNode;
-							FBspNode	&OtherNode	= Model->Nodes[iOtherNode];
-
-							int32 Delta = 
-								(OtherNode.NumVertices + PV2->nVertex - PV1->nVertex) % 
-								OtherNode.NumVertices;
-
-							if( Delta==1 )
-							{
-								// Side is properly linked!
-								int32 OtherVert = PV2->nVertex;
-								int32 iSide;
-								if( Model->Verts[OtherNode.iVertPool+OtherVert].iSide==INDEX_NONE )
-									iSide = Model->NumSharedSides++;
-								else
-									iSide = Model->Verts[OtherNode.iVertPool+OtherVert].iSide;
-									
-								// Link both sides to the shared side.
-								Model->Verts[ Node.iVertPool + ThisVert ].iSide
-								=	Model->Verts[ OtherNode.iVertPool+OtherVert ].iSide
-								=	iSide;
-								goto SkipSide;
-							}
-						}
-					}
-				}
-
-				// This node doesn't have correct side linking
-				//Model->Surfs(Node.iSurf).PolyFlags |= PF_Selected;
-				FBSPOps::GErrors++;
-				//UE_LOG(LogEditorBsp, Warning, "Failed to link side" );
-			}
-
-			// Go to next side.
-			SkipSide:;
-		}
-	}
-	// Gather stats.
-	i=0; int j=0;
-	for( int32 iNode=0; iNode < Model->Nodes.Num(); iNode++ )
-	{
-		FBspNode&	Node		= Model->Nodes[iNode];
-		FVert*		VertPool	= &Model->Verts[Node.iVertPool];
-		
-		i += Node.NumVertices;
-		for( uint8 ThisVert=0; ThisVert < Node.NumVertices; ThisVert++ )
-			if( VertPool[ThisVert].iSide!=INDEX_NONE )
-				j++;
-	}
-	// Done.
-	//UE_LOG(LogEditorBsp, Log,  TEXT("BspOptGeom end") );
-// 	UE_LOG(LogEditorBsp, Log,  TEXT("Processed %i T-points, linked: %i/%i sides"), TeesFound, j, i );
-
-	PointVerts.Free();
-
-	// Remove unused vertices from the vertex streams.
-	// This is necessary to ensure the vertices added to eliminate T junctions
-	// don't overflow the 65536 vertex/stream limit.
-
-	FBSPOps::bspRefresh(Model,0);
+	FBSPUtils::bspOptGeom(Model);
 }
diff --git a/Engine/Source/Editor/UnrealEd/Private/EditorCsg.cpp b/Engine/Source/Editor/UnrealEd/Private/EditorCsg.cpp
--- a/Engine/Source/Editor/UnrealEd/Private/EditorCsg.cpp
+++ b/Engine/Source/Editor/UnrealEd/Private/EditorCsg.cpp
@@ -15,6 +15,7 @@
 #include "EditorModeManager.h"
 #include "SurfaceIterators.h"
 #include "BSPOps.h"
+#include "BSPUtils.h"
 #include "ActorEditorUtils.h"
 #include "Misc/FeedbackContext.h"
 #include "EngineUtils.h"
@@ -115,11 +116,7 @@ protected:
 
 void UEditorEngine::bspRepartition( UWorld* InWorld, int32 iNode )
 {
-	bspBuildFPolys( InWorld->GetModel(), 1, iNode );
-	bspMergeCoplanars( InWorld->GetModel(), 0, 0 );
-	FBSPOps::bspBuild( InWorld->GetModel(), FBSPOps::BSP_Good, 12, 70, 2, iNode );
-	FBSPOps::bspRefresh( InWorld->GetModel(), 1 );
-
+	FBSPUtils::bspRepartition(InWorld, iNode);
 }
 
 //
@@ -325,40 +322,13 @@ void UEditorEngine::csgRebuild( UWorld* InWorld )
 
 void UEditorEngine::polySetAndClearPolyFlags(UModel *Model, uint32 SetBits, uint32 ClearBits,bool SelectedOnly, bool UpdateMaster)
 {
-	for( int32 i=0; i<Model->Surfs.Num(); i++ )
-	{
-		FBspSurf& Poly = Model->Surfs[i];
-		if( !SelectedOnly || (Poly.PolyFlags & PF_Selected) )
-		{
-			uint32 NewFlags = (Poly.PolyFlags & ~ClearBits) | SetBits;
-			if( NewFlags != Poly.PolyFlags )
-			{
-				Model->ModifySurf( i, UpdateMaster );
-				Poly.PolyFlags = NewFlags;
-				if (UpdateMaster)
-				{
-					const bool bUpdateTexCoords = false;
-					const bool bOnlyRefreshSurfaceMaterials = false;
-					polyUpdateMaster(Model, i, bUpdateTexCoords, bOnlyRefreshSurfaceMaterials);
-				}
-			}
-		}
-	}
+	FBSPUtils::polySetAndClearPolyFlags(Model, SetBits, ClearBits, SelectedOnly, UpdateMaster);
 }
 
 
 bool UEditorEngine::polyFindMaster(UModel* InModel, int32 iSurf, FPoly &Poly)
 {
-	FBspSurf &Surf = InModel->Surfs[iSurf];
-	if( !Surf.Actor || !Surf.Actor->Brush->Polys->Element.IsValidIndex(Surf.iBrushPoly) )
-	{
-		return false;
-	}
-	else
-	{
-		Poly = Surf.Actor->Brush->Polys->Element[Surf.iBrushPoly];
-		return true;
-	}
+	return FBSPUtils::polyFindMaster(InModel, iSurf, Poly);
 }
 
 
@@ -370,60 +340,7 @@ void UEditorEngine::polyUpdateMaster
 	bool	bOnlyRefreshSurfaceMaterials
 )
 {
-	FBspSurf &Surf = Model->Surfs[iSurf];
-	ABrush* Actor = Surf.Actor;
-	if( !Actor )
-		return;
-
-	UModel* Brush = Actor->Brush;
-	check(Brush);
-
-	FVector ActorLocation;
-	FVector ActorPrePivot;
-	FVector ActorScale;
-	FRotator ActorRotation;
-
-	if (Brush->bCachedOwnerTransformValid)
-	{
-		// Use transform cached when the geometry was last built, in case the current Actor transform has changed since then
-		// (e.g. because Auto Update BSP is disabled)
-		ActorLocation = Brush->OwnerLocationWhenLastBuilt;
-		ActorScale = Brush->OwnerScaleWhenLastBuilt;
-		ActorRotation = Brush->OwnerRotationWhenLastBuilt;
-	}
-	else
-	{
-		// No cached owner transform, so use the current one
-		ActorLocation = Actor->GetActorLocation();
-		ActorScale = Actor->GetActorScale();
-		ActorRotation = Actor->GetActorRotation();
-	}
-
-	const FRotationMatrix RotationMatrix(ActorRotation);
-
-	for (int32 iEdPoly = Surf.iBrushPoly; iEdPoly < Brush->Polys->Element.Num(); iEdPoly++)
-	{
-		FPoly& MasterEdPoly = Brush->Polys->Element[iEdPoly];
-		if (iEdPoly == Surf.iBrushPoly || MasterEdPoly.iLink == Surf.iBrushPoly)
-		{
-			MasterEdPoly.Material = Surf.Material;
-			MasterEdPoly.PolyFlags = Surf.PolyFlags & ~(PF_NoEdit);
-
-			if (bUpdateTexCoords)
-			{
-				MasterEdPoly.Base = RotationMatrix.InverseTransformVector(Model->Points[Surf.pBase] - ActorLocation) / ActorScale;
-				MasterEdPoly.TextureU = RotationMatrix.InverseTransformVector(Model->Vectors[Surf.vTextureU]) * ActorScale;
-				MasterEdPoly.TextureV = RotationMatrix.InverseTransformVector(Model->Vectors[Surf.vTextureV]) * ActorScale;
-			}
-		}
-	}
-
-	Model->InvalidSurfaces = true;
-
-	if (bOnlyRefreshSurfaceMaterials)
-	{
-		Model->bOnlyRebuildMaterialIndexBuffers = true;
-	}
+	FBSPUtils::polyUpdateMaster(Model, iSurf, bUpdateTexCoords, bOnlyRefreshSurfaceMaterials);
 }
 
 
@@ -434,77 +351,13 @@ void UEditorEngine::polyGetLinkedPolys
 	TArray<FPoly>* InPolyList
 )
 {
-	InPolyList->Empty();
-
-	if( InPoly->iLink == INDEX_NONE )
-	{
-		// If this poly has no links, just stick the one poly in the final list.
-		new(*InPolyList)FPoly( *InPoly );
-	}
-	else
-	{
-		// Find all polys that match the source polys link value.
-		for( int32 poly = 0 ; poly < InBrush->Brush->Polys->Element.Num() ; poly++ )
-			if( InBrush->Brush->Polys->Element[poly].iLink == InPoly->iLink )
-				new(*InPolyList)FPoly( InBrush->Brush->Polys->Element[poly] );
-	}
-
+	FBSPUtils::polyGetLinkedPolys(InBrush, InPoly, InPolyList);
 }
 
 
 void UEditorEngine::polySplitOverlappingEdges( TArray<FPoly>* InPolyList, TArray<FPoly>* InResult )
 {
-	InResult->Empty();
-
-	for( int32 poly = 0 ; poly < InPolyList->Num() ; poly++ )
-	{
-		FPoly* SrcPoly = &(*InPolyList)[poly];
-		FPoly NewPoly = *SrcPoly;
-
-		for( int32 edge = 0 ; edge < SrcPoly->Vertices.Num() ; edge++ )
-		{
-			FEdge SrcEdge = FEdge( SrcPoly->Vertices[edge], SrcPoly->Vertices[ edge+1 < SrcPoly->Vertices.Num() ? edge+1 : 0 ] );
-			FPlane SrcEdgePlane( SrcEdge.Vertex[0], SrcEdge.Vertex[1], SrcEdge.Vertex[0] + (SrcPoly->Normal * 16) );
-
-			for( int32 poly2 = 0 ; poly2 < InPolyList->Num() ; poly2++ )
-			{
-				FPoly* CmpPoly = &(*InPolyList)[poly2];
-
-				// We can't compare to ourselves.
-				if( CmpPoly == SrcPoly )
-					continue;
-
-				for( int32 edge2 = 0 ; edge2 < CmpPoly->Vertices.Num() ; edge2++ )
-				{
-					FEdge CmpEdge = FEdge( CmpPoly->Vertices[edge2], CmpPoly->Vertices[ edge2+1 < CmpPoly->Vertices.Num() ? edge2+1 : 0 ] );
-
-					// If both vertices on this edge lie on the same plane as the original edge, create
-					// a sphere around the original 2 vertices.  If either of this edges vertices are inside of
-					// that sphere, we need to split the original edge by adding a vertex to it's poly.
-					if( FMath::Abs( FVector::PointPlaneDist( CmpEdge.Vertex[0], SrcEdge.Vertex[0], SrcEdgePlane ) ) < THRESH_POINT_ON_PLANE
-							&& FMath::Abs( FVector::PointPlaneDist( CmpEdge.Vertex[1], SrcEdge.Vertex[0], SrcEdgePlane ) ) < THRESH_POINT_ON_PLANE )
-					{
-						//
-						// Check THIS edge against the SOURCE edge
-						//
-
-						FVector Dir = SrcEdge.Vertex[1] - SrcEdge.Vertex[0];
-						Dir.Normalize();
-						float Dist = FVector::Dist( SrcEdge.Vertex[1], SrcEdge.Vertex[0] );
-						FVector Origin = SrcEdge.Vertex[0] + (Dir * (Dist / 2.0f));
-						float Radius = Dist / 2.0f;
-
-						for( int32 vtx = 0 ; vtx < 2 ; vtx++ )
-							if( FVector::Dist( Origin, CmpEdge.Vertex[vtx] ) && FVector::Dist( Origin, CmpEdge.Vertex[vtx] ) < Radius )
-								NewPoly.InsertVertex( edge2+1, CmpEdge.Vertex[vtx] );
-					}
-				}
-			}
-		}
-
-		new(*InResult)FPoly( NewPoly );
-	}
-
+	FBSPUtils::polySplitOverlappingEdges(InPolyList, InResult);
 }
 
 
@@ -514,63 +367,7 @@ void UEditorEngine::polyGetOuterEdgeList
 	TArray<FEdge>* InEdgeList
 )
 {
-	TArray<FPoly> NewPolyList;
-	polySplitOverlappingEdges( InPolyList, &NewPolyList );
-
-	TArray<FEdge> TempEdges;
-
-	// Create a master list of edges.
-	for( int32 poly = 0 ; poly < NewPolyList.Num() ; poly++ )
-	{
-		FPoly* Poly = &NewPolyList[poly];
-		for( int32 vtx = 0 ; vtx < Poly->Vertices.Num() ; vtx++ )
-			new( TempEdges )FEdge( Poly->Vertices[vtx], Poly->Vertices[ vtx+1 < Poly->Vertices.Num() ? vtx+1 : 0] );
-	}
-
-	// Add all the unique edges into the final edge list.
-	TArray<FEdge> FinalEdges;
-
-	for( int32 tedge = 0 ; tedge < TempEdges.Num() ; tedge++ )
-	{
-		FEdge* TestEdge = &TempEdges[tedge];
-
-		int32 EdgeCount = 0;
-		for( int32 edge = 0 ; edge < TempEdges.Num() ; edge++ )
-		{
-			if( TempEdges[edge] == *TestEdge )
-				EdgeCount++;
-		}
-
-		if( EdgeCount == 1 )
-			new( FinalEdges )FEdge( *TestEdge );
-	}
-
-	// Reorder all the edges so that they line up, end to end.
-	InEdgeList->Empty();
-	if( !FinalEdges.Num() ) return;
-
-	new( *InEdgeList )FEdge( FinalEdges[0] );
-	FVector Comp = FinalEdges[0].Vertex[1];
-	FinalEdges.RemoveAt(0);
-
-	FEdge DebuG;
-	for( int32 x = 0 ; x < FinalEdges.Num() ; x++ )
-	{
-		DebuG = FinalEdges[x];
-
-		// If the edge is backwards, flip it
-		if( FinalEdges[x].Vertex[1] == Comp )
-			Exchange( FinalEdges[x].Vertex[0], FinalEdges[x].Vertex[1] );
-
-		if( FinalEdges[x].Vertex[0] == Comp )
-		{
-			new( *InEdgeList )FEdge( FinalEdges[x] );
-			Comp = FinalEdges[x].Vertex[1];
-			FinalEdges.RemoveAt(x);
-			x = -1;
-		}
-	}
-
+	FBSPUtils::polyGetOuterEdgeList(InPolyList, InEdgeList);
 }
 
 /*-----------------------------------------------------------------------------
diff --git a/Engine/Source/Editor/UnrealEd/Private/EditorServer.cpp b/Engine/Source/Editor/UnrealEd/Private/EditorServer.cpp
--- a/Engine/Source/Editor/UnrealEd/Private/EditorServer.cpp
+++ b/Engine/Source/Editor/UnrealEd/Private/EditorServer.cpp
@@ -105,6 +105,7 @@
 #include "ScopedTransaction.h"
 #include "SurfaceIterators.h"
 #include "LightMap.h"
+#include "Engine/BrushBuilder.h"
 #include "BSPOps.h"
 #include "EditorLevelUtils.h"
 #include "Interfaces/IMainFrameModule.h"
@@ -755,6 +756,11 @@ bool UEditorEngine::Exec_Brush( UWorld* InWorld, const TCHAR* Str, FOutputDevice
 			NewBrush = FBSPOps::csgAddOperation( WorldBrush, DWord1, Brush_Add );
 			if( NewBrush )
 			{
+				if( NewBrush->GetBrushBuilder() )
+				{
+					FActorLabelUtilities::SetActorLabelUnique(NewBrush, FText::Format(NSLOCTEXT("UnrealEd", "BrushName", "{0} Brush"), FText::FromString(NewBrush->GetBrushBuilder()->GetClass()->GetDescription())).ToString());
+				}
+
 				// Materials selected in the Content Browser, but not actually loaded, will not be
 				// in the global selection set, which is expected by bspBrushCSG when it comes to
 				// applying the material to the surfaces. This goes through the set of objects selected
@@ -853,6 +859,11 @@ bool UEditorEngine::Exec_Brush( UWorld* InWorld, const TCHAR* Str, FOutputDevice
 			NewBrush = FBSPOps::csgAddOperation(WorldBrush,0,Brush_Subtract); // Layer
 			if( NewBrush )
 			{
+				if( NewBrush->GetBrushBuilder() )
+				{
+					FActorLabelUtilities::SetActorLabelUnique(NewBrush, FText::Format(NSLOCTEXT("UnrealEd", "BrushName", "{0} Brush"), FText::FromString(NewBrush->GetBrushBuilder()->GetClass()->GetDescription())).ToString());
+				}
+
 				NewBrush->Modify();
 				InWorld->GetModel()->Modify();
 				bspBrushCSG( NewBrush, InWorld->GetModel(), 0, Brush_Subtract, CSG_None, true, true, true );
diff --git a/Engine/Source/Editor/UnrealEd/Public/BSPOps.h b/Engine/Source/Editor/UnrealEd/Public/BSPOps.h
deleted file mode 100644
diff --git a/Engine/Source/Editor/UnrealEd/UnrealEd.Build.cs b/Engine/Source/Editor/UnrealEd/UnrealEd.Build.cs
--- a/Engine/Source/Editor/UnrealEd/UnrealEd.Build.cs
+++ b/Engine/Source/Editor/UnrealEd/UnrealEd.Build.cs
@@ -29,6 +29,7 @@ public class UnrealEd : ModuleRules
 		PrivateIncludePathModuleNames.AddRange(
 			new string[]
 			{
+				"BSPUtils",
 				"BehaviorTreeEditor",
 				"ClassViewer",
 				"StructViewer",
@@ -126,6 +127,7 @@ public class UnrealEd : ModuleRules
 			{
 				"AssetRegistry",
 				"AssetTagsEditor",
+				"BSPUtils",
 				"LevelSequence",
 				"AnimGraph",
 				"AppFramework",
