BSP Pen Tool: Fix issue with Enter not closing polygon due to default brush not having any geometry

diff --git a/Engine/Plugins/Editor/GeometryMode/Source/GeometryMode/Private/GeometryModifiers.cpp b/Engine/Plugins/Editor/GeometryMode/Source/GeometryMode/Private/GeometryModifiers.cpp
--- a/Engine/Plugins/Editor/GeometryMode/Source/GeometryMode/Private/GeometryModifiers.cpp
+++ b/Engine/Plugins/Editor/GeometryMode/Source/GeometryMode/Private/GeometryModifiers.cpp
@@ -1352,7 +1352,7 @@ void UGeomModifier_Pen::Apply()
 		}
 		else
 		{
-			ResultingBrush = FBSPOps::csgAddOperation( BuilderBrush, PF_DefaultFlags, Brush_Add );
+			ResultingBrush = FBSPOps::csgAddOperation( BuilderBrush, PF_DefaultFlags, Brush_Add, true );
 			if (ResultingBrush == nullptr)
 			{
 				return;
diff --git a/Engine/Source/Developer/BSPUtils/Private/BSPOps.cpp b/Engine/Source/Developer/BSPUtils/Private/BSPOps.cpp
--- a/Engine/Source/Developer/BSPUtils/Private/BSPOps.cpp
+++ b/Engine/Source/Developer/BSPUtils/Private/BSPOps.cpp
@@ -587,7 +587,7 @@ void FBSPOps::csgCopyBrush( ABrush* Dest, ABrush* Src, uint32 PolyFlags, EObject
  *
  * @return		A newly-created copy of the brush.
  */
-ABrush*	FBSPOps::csgAddOperation( ABrush* Actor, uint32 PolyFlags, EBrushType BrushType)
+ABrush*	FBSPOps::csgAddOperation( ABrush* Actor, uint32 PolyFlags, EBrushType BrushType, bool bAllowEmpty)
 {
 	check(Actor);
 	check(Actor->GetBrushComponent());
@@ -596,7 +596,7 @@ ABrush*	FBSPOps::csgAddOperation( ABrush* Actor, uint32 PolyFlags, EBrushType Br
 	check(Actor->GetWorld());
 
 	// Can't do this if brush has no polys. 
-	if( !Actor->Brush->Polys->Element.Num() )
+	if(!bAllowEmpty && !Actor->Brush->Polys->Element.Num())
 		return NULL;
 
 	// Spawn a new actor for the brush.
@@ -612,6 +612,7 @@ ABrush*	FBSPOps::csgAddOperation( ABrush* Actor, uint32 PolyFlags, EBrushType Br
 		PolyFlags,
 		RF_Transactional,
 		0,
+		true,
 		true
 	);
 	check(Result->Brush);
diff --git a/Engine/Source/Developer/BSPUtils/Public/BSPOps.h b/Engine/Source/Developer/BSPUtils/Public/BSPOps.h
--- a/Engine/Source/Developer/BSPUtils/Public/BSPOps.h
+++ b/Engine/Source/Developer/BSPUtils/Public/BSPOps.h
@@ -31,7 +31,7 @@ public:
 
 	BSPUTILS_API static void csgPrepMovingBrush( ABrush* Actor );
 	BSPUTILS_API static void csgCopyBrush( ABrush* Dest, ABrush* Src, uint32 PolyFlags, EObjectFlags ResFlags, bool bNeedsPrep, bool bCopyPosRotScale, bool bAllowEmpty = false );
-	BSPUTILS_API static ABrush*	csgAddOperation( ABrush* Actor, uint32 PolyFlags, EBrushType BrushType );
+	BSPUTILS_API static ABrush*	csgAddOperation( ABrush* Actor, uint32 PolyFlags, EBrushType BrushType, bool bAllowEmpty = false );
 	BSPUTILS_API static int32 bspAddVector( UModel* Model, const FVector* V, bool Exact );
 	BSPUTILS_API static int32 bspAddPoint( UModel* Model, const FVector* V, bool Exact );
 	BSPUTILS_API static void bspBuild( UModel* Model, enum EBspOptimization Opt, int32 Balance, int32 PortalBias, int32 RebuildSimplePolys, int32 iNode );
