FBSPUtils: Resolve issues with non-inclusive terminology

diff --git a/Engine/Plugins/Editor/GeometryMode/Source/TextureAlignMode/Private/TextureAlignEdMode.cpp b/Engine/Plugins/Editor/GeometryMode/Source/TextureAlignMode/Private/TextureAlignEdMode.cpp
--- a/Engine/Plugins/Editor/GeometryMode/Source/TextureAlignMode/Private/TextureAlignEdMode.cpp
+++ b/Engine/Plugins/Editor/GeometryMode/Source/TextureAlignMode/Private/TextureAlignEdMode.cpp
@@ -270,7 +270,7 @@ bool FModeTool_Texture::InputDelta(FEditorViewportClient* InViewportClient,FView
 
 			const bool bUpdateTexCoords = true;
 			const bool bOnlyRefreshSurfaceMaterials = true;
-			GEditor->polyUpdateMaster(Model, It.GetSurfaceIndex(), bUpdateTexCoords, bOnlyRefreshSurfaceMaterials);
+			GEditor->polyUpdateBrush(Model, It.GetSurfaceIndex(), bUpdateTexCoords, bOnlyRefreshSurfaceMaterials);
 		}
 	}
 
diff --git a/Engine/Source/Developer/BSPUtils/Private/BSPUtils.cpp b/Engine/Source/Developer/BSPUtils/Private/BSPUtils.cpp
--- a/Engine/Source/Developer/BSPUtils/Private/BSPUtils.cpp
+++ b/Engine/Source/Developer/BSPUtils/Private/BSPUtils.cpp
@@ -931,7 +931,7 @@ int FBSPUtils::bspNodeToFPoly
 	FPoly	*EdPoly
 )
 {
-	FPoly MasterEdPoly;
+	FPoly BrushEdPoly;
 
 	FBspNode &Node     	= Model->Nodes[iNode];
 	FBspSurf &Poly     	= Model->Surfs[Node.iSurf];
@@ -947,8 +947,8 @@ int FBSPUtils::bspNodeToFPoly
 	EdPoly->Actor    	= Poly.Actor;
 	EdPoly->iBrushPoly  = Poly.iBrushPoly;
 	
-	if( polyFindMaster(Model,Node.iSurf,MasterEdPoly) )
-		EdPoly->ItemName  = MasterEdPoly.ItemName;
+	if( polyFindBrush(Model,Node.iSurf,BrushEdPoly) )
+		EdPoly->ItemName  = BrushEdPoly.ItemName;
 	else
 		EdPoly->ItemName  = NAME_None;
 
@@ -1922,7 +1922,7 @@ void FBSPUtils::bspRepartition( UWorld* InWorld, int32 iNode )
 }
 
 
-void FBSPUtils::polySetAndClearPolyFlags(UModel *Model, uint32 SetBits, uint32 ClearBits,bool SelectedOnly, bool UpdateMaster)
+void FBSPUtils::polySetAndClearPolyFlags(UModel *Model, uint32 SetBits, uint32 ClearBits,bool SelectedOnly, bool UpdateBrush)
 {
 	for( int32 i=0; i<Model->Surfs.Num(); i++ )
 	{
@@ -1932,13 +1932,13 @@ void FBSPUtils::polySetAndClearPolyFlags(UModel *Model, uint32 SetBits, uint32 C
 			uint32 NewFlags = (Poly.PolyFlags & ~ClearBits) | SetBits;
 			if( NewFlags != Poly.PolyFlags )
 			{
-				Model->ModifySurf( i, UpdateMaster );
+				Model->ModifySurf( i, UpdateBrush );
 				Poly.PolyFlags = NewFlags;
-				if (UpdateMaster)
+				if (UpdateBrush)
 				{
 					const bool bUpdateTexCoords = false;
 					const bool bOnlyRefreshSurfaceMaterials = false;
-					polyUpdateMaster(Model, i, bUpdateTexCoords, bOnlyRefreshSurfaceMaterials);
+					polyUpdateBrush(Model, i, bUpdateTexCoords, bOnlyRefreshSurfaceMaterials);
 				}
 			}
 		}
@@ -1947,6 +1947,11 @@ void FBSPUtils::polySetAndClearPolyFlags(UModel *Model, uint32 SetBits, uint32 C
 
 
 bool FBSPUtils::polyFindMaster(UModel* InModel, int32 iSurf, FPoly &Poly)
+{
+	return polyFindBrush(InModel, iSurf, Poly);
+}
+
+bool FBSPUtils::polyFindBrush(UModel* InModel, int32 iSurf, FPoly &Poly)
 {
 	FBspSurf &Surf = InModel->Surfs[iSurf];
 	if( !Surf.Actor || !Surf.Actor->Brush->Polys->Element.IsValidIndex(Surf.iBrushPoly) )
@@ -1961,7 +1966,12 @@ bool FBSPUtils::polyFindMaster(UModel* InModel, int32 iSurf, FPoly &Poly)
 }
 
 
-void FBSPUtils::polyUpdateMaster
+void FBSPUtils::polyUpdateMaster(UModel* Model, int32 iSurf, bool bUpdateTexCoords,	bool bOnlyRefreshSurfaceMaterials)
+{
+	polyUpdateBrush(Model, iSurf, bUpdateTexCoords,	bOnlyRefreshSurfaceMaterials);
+}
+
+void FBSPUtils::polyUpdateBrush
 (
 	UModel*	Model,
 	int32  	iSurf,
@@ -2002,17 +2012,17 @@ void FBSPUtils::polyUpdateMaster
 
 	for (int32 iEdPoly = Surf.iBrushPoly; iEdPoly < Brush->Polys->Element.Num(); iEdPoly++)
 	{
-		FPoly& MasterEdPoly = Brush->Polys->Element[iEdPoly];
-		if (iEdPoly == Surf.iBrushPoly || MasterEdPoly.iLink == Surf.iBrushPoly)
+		FPoly& BrushEdPoly = Brush->Polys->Element[iEdPoly];
+		if (iEdPoly == Surf.iBrushPoly || BrushEdPoly.iLink == Surf.iBrushPoly)
 		{
-			MasterEdPoly.Material = Surf.Material;
-			MasterEdPoly.PolyFlags = Surf.PolyFlags & ~(PF_NoEdit);
+			BrushEdPoly.Material = Surf.Material;
+			BrushEdPoly.PolyFlags = Surf.PolyFlags & ~(PF_NoEdit);
 
 			if (bUpdateTexCoords)
 			{
-				MasterEdPoly.Base = FVector3f(RotationMatrix.InverseTransformVector((FVector)Model->Points[Surf.pBase] - ActorLocation) / ActorScale);
-				MasterEdPoly.TextureU = FVector3f(RotationMatrix.InverseTransformVector((FVector)Model->Vectors[Surf.vTextureU]) * ActorScale);
-				MasterEdPoly.TextureV = FVector3f(RotationMatrix.InverseTransformVector((FVector)Model->Vectors[Surf.vTextureV]) * ActorScale);
+				BrushEdPoly.Base = FVector3f(RotationMatrix.InverseTransformVector((FVector)Model->Points[Surf.pBase] - ActorLocation) / ActorScale);
+				BrushEdPoly.TextureU = FVector3f(RotationMatrix.InverseTransformVector((FVector)Model->Vectors[Surf.vTextureU]) * ActorScale);
+				BrushEdPoly.TextureV = FVector3f(RotationMatrix.InverseTransformVector((FVector)Model->Vectors[Surf.vTextureV]) * ActorScale);
 			}
 		}
 	}
@@ -2117,7 +2127,7 @@ void FBSPUtils::polyGetOuterEdgeList
 
 	TArray<FEdge> TempEdges;
 
-	// Create a master list of edges.
+	// Create a list of edges.
 	for( int32 poly = 0 ; poly < NewPolyList.Num() ; poly++ )
 	{
 		FPoly* Poly = &NewPolyList[poly];
diff --git a/Engine/Source/Developer/BSPUtils/Public/BSPUtils.h b/Engine/Source/Developer/BSPUtils/Public/BSPUtils.h
--- a/Engine/Source/Developer/BSPUtils/Public/BSPUtils.h
+++ b/Engine/Source/Developer/BSPUtils/Public/BSPUtils.h
@@ -51,7 +51,7 @@ public:
 	 * Sets and clears all Bsp node flags.  Affects all nodes, even ones that don't
 	 * really exist.
 	 */
-	static void polySetAndClearPolyFlags(UModel *Model, uint32 SetBits, uint32 ClearBits,bool SelectedOnly, bool UpdateMaster);
+	static void polySetAndClearPolyFlags(UModel *Model, uint32 SetBits, uint32 ClearBits,bool SelectedOnly, bool UpdateBrush);
 
 	/**
 	 *
@@ -63,14 +63,20 @@ public:
 	 *
 	 * returns true if poly not available
 	 */
+	static bool polyFindBrush(UModel* InModel, int32 iSurf, FPoly &Poly);
+
+	UE_DEPRECATED(5.1, "polyFindMaster is deprecated; please use polyFindBrush instead")
 	static bool polyFindMaster(UModel* InModel, int32 iSurf, FPoly &Poly);
 
 	/**
-	 * Update a the master brush EdPoly corresponding to a newly-changed
+	 * Update the brush EdPoly corresponding to a newly-changed
 	 * poly to reflect its new properties.
 	 *
 	 * Doesn't do any transaction tracking.
 	 */
+	static void polyUpdateBrush(UModel* Model, int32 iSurf, bool bUpdateTexCoords, bool bOnlyRefreshSurfaceMaterials);
+	
+	UE_DEPRECATED(5.1, "polyUpdateMaster is deprecated; please use polyUpdateBrush instead")
 	static void polyUpdateMaster(UModel* Model, int32 iSurf, bool bUpdateTexCoords, bool bOnlyRefreshSurfaceMaterials);
 	
 	/**
diff --git a/Engine/Source/Editor/LevelEditor/Private/LevelEditorGenericDetails.cpp b/Engine/Source/Editor/LevelEditor/Private/LevelEditorGenericDetails.cpp
--- a/Engine/Source/Editor/LevelEditor/Private/LevelEditorGenericDetails.cpp
+++ b/Engine/Source/Editor/LevelEditor/Private/LevelEditorGenericDetails.cpp
@@ -131,7 +131,7 @@ void FLevelEditorGenericDetails::OnMaterialChanged( UMaterialInterface* NewMater
 			Model->Surfs[SurfaceIndex].Material = NewMaterial;
 			const bool bUpdateTexCoords = false;
 			const bool bOnlyUpdateSurfaceMaterials = true;
-			GEditor->polyUpdateMaster(Model, SurfaceIndex, bUpdateTexCoords, bOnlyUpdateSurfaceMaterials);
+			GEditor->polyUpdateBrush(Model, SurfaceIndex, bUpdateTexCoords, bOnlyUpdateSurfaceMaterials);
 			Model->MarkPackageDirty();
 
 			bModelDirtied = true;
diff --git a/Engine/Source/Editor/UnrealEd/Classes/Editor/EditorEngine.h b/Engine/Source/Editor/UnrealEd/Classes/Editor/EditorEngine.h
--- a/Engine/Source/Editor/UnrealEd/Classes/Editor/EditorEngine.h
+++ b/Engine/Source/Editor/UnrealEd/Classes/Editor/EditorEngine.h
@@ -1389,14 +1389,20 @@ public:
 	*
 	* returns true if poly not available
 	*/
+	static bool polyFindBrush(UModel* InModel, int32 iSurf, FPoly &Poly);
+
+	UE_DEPRECATED(5.1, "polyFindMaster is deprecated; please use polyFindBrush instead")
 	virtual bool polyFindMaster( UModel* InModel, int32 iSurf, FPoly& Poly );
 
 	/**
-	 * Update a the master brush EdPoly corresponding to a newly-changed
+	 * Update a the brush EdPoly corresponding to a newly-changed
 	 * poly to reflect its new properties.
 	 *
 	 * Doesn't do any transaction tracking.
 	 */
+	static void polyUpdateBrush(UModel* Model, int32 iSurf, bool bUpdateTexCoords, bool bOnlyRefreshSurfaceMaterials);
+
+	UE_DEPRECATED(5.1, "polyUpdateMaster is deprecated; please use polyUpdateBrush instead")
 	virtual void polyUpdateMaster( UModel* Model, int32 iSurf, bool bUpdateTexCoords, bool bOnlyRefreshSurfaceMaterials );
 
 	/**
@@ -1421,7 +1427,7 @@ public:
 	 * Sets and clears all Bsp node flags.  Affects all nodes, even ones that don't
 	 * really exist.
 	 */
-	virtual void polySetAndClearPolyFlags( UModel* Model, uint32 SetBits, uint32 ClearBits, bool SelectedOnly, bool UpdateMaster );
+	virtual void polySetAndClearPolyFlags( UModel* Model, uint32 SetBits, uint32 ClearBits, bool SelectedOnly, bool UpdateBrush );
 
 	// Selection.
 	virtual void SelectActor(AActor* Actor, bool bInSelected, bool bNotify, bool bSelectEvenIfHidden = false, bool bForceRefresh = false) {}
diff --git a/Engine/Source/Editor/UnrealEd/Private/EditorCsg.cpp b/Engine/Source/Editor/UnrealEd/Private/EditorCsg.cpp
--- a/Engine/Source/Editor/UnrealEd/Private/EditorCsg.cpp
+++ b/Engine/Source/Editor/UnrealEd/Private/EditorCsg.cpp
@@ -320,19 +320,29 @@ void UEditorEngine::csgRebuild( UWorld* InWorld )
 }
 
 
-void UEditorEngine::polySetAndClearPolyFlags(UModel *Model, uint32 SetBits, uint32 ClearBits,bool SelectedOnly, bool UpdateMaster)
+void UEditorEngine::polySetAndClearPolyFlags(UModel *Model, uint32 SetBits, uint32 ClearBits,bool SelectedOnly, bool UpdateBrush)
 {
-	FBSPUtils::polySetAndClearPolyFlags(Model, SetBits, ClearBits, SelectedOnly, UpdateMaster);
+	FBSPUtils::polySetAndClearPolyFlags(Model, SetBits, ClearBits, SelectedOnly, UpdateBrush);
 }
 
 
 bool UEditorEngine::polyFindMaster(UModel* InModel, int32 iSurf, FPoly &Poly)
 {
-	return FBSPUtils::polyFindMaster(InModel, iSurf, Poly);
+	return polyFindBrush(InModel, iSurf, Poly);
 }
 
+bool UEditorEngine::polyFindBrush(UModel* InModel, int32 iSurf, FPoly &Poly)
+{
+	return FBSPUtils::polyFindBrush(InModel, iSurf, Poly);
+}
+
+
+void UEditorEngine::polyUpdateMaster(UModel* Model, int32 iSurf, bool bUpdateTexCoords, bool bOnlyRefreshSurfaceMaterials)
+{
+	polyUpdateBrush(Model, iSurf, bUpdateTexCoords, bOnlyRefreshSurfaceMaterials);
+}
 
-void UEditorEngine::polyUpdateMaster
+void UEditorEngine::polyUpdateBrush
 (
 	UModel*	Model,
 	int32  	iSurf,
@@ -340,7 +350,7 @@ void UEditorEngine::polyUpdateMaster
 	bool	bOnlyRefreshSurfaceMaterials
 )
 {
-	FBSPUtils::polyUpdateMaster(Model, iSurf, bUpdateTexCoords, bOnlyRefreshSurfaceMaterials);
+	FBSPUtils::polyUpdateBrush(Model, iSurf, bUpdateTexCoords, bOnlyRefreshSurfaceMaterials);
 }
 
 
@@ -424,7 +434,7 @@ void UEditorEngine::polySelectMatchingGroups( UModel* Model )
 		FBspSurf *Surf = &Model->Surfs(i);
 		if( Surf->PolyFlags&PF_Selected )
 		{
-			FPoly Poly; polyFindMaster(Model,i,Poly);
+			FPoly Poly; polyFindBrush(Model,i,Poly);
 			GFlags1[Poly.Actor->Layer.GetIndex()]=1;
 		}
 	}
@@ -432,7 +442,7 @@ void UEditorEngine::polySelectMatchingGroups( UModel* Model )
 	{
 		FBspSurf *Surf = &Model->Surfs(i);
 		FPoly Poly;
-		polyFindMaster(Model,i,Poly);
+		polyFindBrush(Model,i,Poly);
 		if
 		(	(GFlags1[Poly.Actor->Layer.GetIndex()]) 
 			&&	(!(Surf->PolyFlags & PF_Selected)) )
@@ -461,7 +471,7 @@ void UEditorEngine::polySelectMatchingItems(UModel *InModel)
 		}
 		if( Surf->PolyFlags&PF_Selected )
 		{
-			FPoly Poly; polyFindMaster(InModel,i,Poly);
+			FPoly Poly; polyFindBrush(InModel,i,Poly);
 			GFlags1[Poly.ItemName.GetIndex()]=1;
 		}
 	}
@@ -470,7 +480,7 @@ void UEditorEngine::polySelectMatchingItems(UModel *InModel)
 		FBspSurf *Surf = &InModel->Surfs(i);
 		if( Surf->Actor )
 		{
-			FPoly Poly; polyFindMaster(InModel,i,Poly);
+			FPoly Poly; polyFindBrush(InModel,i,Poly);
 			if ((GFlags1[Poly.ItemName.GetIndex()]) &&
 				( GFlags2[Surf->Actor->Brush->GetIndex()]) &&
 				(!(Surf->PolyFlags & PF_Selected)))
@@ -1280,7 +1290,7 @@ void UEditorEngine::polyTexPan(UModel *Model,int32 PanU,int32 PanV,int32 Absolut
 
 			const bool bUpdateTexCoords = true;
 			const bool bOnlyRefreshSurfaceMaterials = true;
-			polyUpdateMaster(Model, SurfaceIndex, bUpdateTexCoords, bOnlyRefreshSurfaceMaterials);
+			polyUpdateBrush(Model, SurfaceIndex, bUpdateTexCoords, bOnlyRefreshSurfaceMaterials);
 		}
 	}
 }
@@ -1309,7 +1319,7 @@ void UEditorEngine::polyTexScale( UModel* Model, float UU, float UV, float VU, f
 			// Update generating brush poly.
 			const bool bUpdateTexCoords = true;
 			const bool bOnlyRefreshSurfaceMaterials = true;
-			polyUpdateMaster(Model, i, bUpdateTexCoords, bOnlyRefreshSurfaceMaterials);
+			polyUpdateBrush(Model, i, bUpdateTexCoords, bOnlyRefreshSurfaceMaterials);
 		}
 	}
 }
diff --git a/Engine/Source/Editor/UnrealEd/Private/EditorExporters.cpp b/Engine/Source/Editor/UnrealEd/Private/EditorExporters.cpp
--- a/Engine/Source/Editor/UnrealEd/Private/EditorExporters.cpp
+++ b/Engine/Source/Editor/UnrealEd/Private/EditorExporters.cpp
@@ -1378,7 +1378,7 @@ static void ExportPolys( UPolys* Polys, int32 &PolyNum, int32 TotalPolys, FFeedb
 			const FVector& Normal = (FVector)Model->Vectors[Surf.vNormal];
 
 			FPoly Poly;
-			GEditor->polyFindMaster( Model, Node->iSurf, Poly );
+			GEditor->polyFindBrush( Model, Node->iSurf, Poly );
 
 			// Triangulate this node and generate an OBJ face from the vertices.
 			for(int32 StartVertexIndex = 1;StartVertexIndex < Node->NumVertices-1;StartVertexIndex++)
@@ -2283,7 +2283,7 @@ void UEditorEngine::RebuildStaticNavigableGeometry(ULevel* Level)
 				const FVector& Normal = (FVector)Model->Vectors[Surf.vNormal];
 
 				FPoly Poly;
-				polyFindMaster( Model, Node->iSurf, Poly );
+				polyFindBrush( Model, Node->iSurf, Poly );
 
 				// Triangulate this node and generate an OBJ face from the vertices.
 				for(int32 StartVertexIndex = 1;StartVertexIndex < Node->NumVertices-1;StartVertexIndex++)
diff --git a/Engine/Source/Editor/UnrealEd/Private/EditorServer.cpp b/Engine/Source/Editor/UnrealEd/Private/EditorServer.cpp
--- a/Engine/Source/Editor/UnrealEd/Private/EditorServer.cpp
+++ b/Engine/Source/Editor/UnrealEd/Private/EditorServer.cpp
@@ -4510,7 +4510,7 @@ bool UEditorEngine::Exec_Poly( UWorld* InWorld, const TCHAR* Str, FOutputDevice&
 				Model->Surfs[SurfaceIndex].Material = SelectedMaterialInstance;
 				const bool bUpdateTexCoords = false;
 				const bool bOnlyRefreshSurfaceMaterials = true;
-				polyUpdateMaster(Model, SurfaceIndex, bUpdateTexCoords, bOnlyRefreshSurfaceMaterials);
+				polyUpdateBrush(Model, SurfaceIndex, bUpdateTexCoords, bOnlyRefreshSurfaceMaterials);
 				Model->MarkPackageDirty();
 
 				bModelDirtied = true;
@@ -4543,7 +4543,7 @@ bool UEditorEngine::Exec_Poly( UWorld* InWorld, const TCHAR* Str, FOutputDevice&
 					It.GetModel()->Surfs[SurfaceIndex].Material = Material;
 					const bool bUpdateTexCoords = false;
 					const bool bOnlyRefreshSurfaceMaterials = true;
-					polyUpdateMaster(It.GetModel(), SurfaceIndex, bUpdateTexCoords, bOnlyRefreshSurfaceMaterials);
+					polyUpdateBrush(It.GetModel(), SurfaceIndex, bUpdateTexCoords, bOnlyRefreshSurfaceMaterials);
 				}
 			}
 
diff --git a/Engine/Source/Editor/UnrealEd/Private/Factories/EditorFactories.cpp b/Engine/Source/Editor/UnrealEd/Private/Factories/EditorFactories.cpp
--- a/Engine/Source/Editor/UnrealEd/Private/Factories/EditorFactories.cpp
+++ b/Engine/Source/Editor/UnrealEd/Private/Factories/EditorFactories.cpp
@@ -1172,7 +1172,7 @@ UObject* ULevelFactory::FactoryCreateText
 							const bool bOnlyRefreshSurfaceMaterials = false;
 							if (GWorld == World)
 							{
-								GEditor->polyUpdateMaster(CurrentLevel->Model, i, bUpdateTexCoords, bOnlyRefreshSurfaceMaterials);
+								GEditor->polyUpdateBrush(CurrentLevel->Model, i, bUpdateTexCoords, bOnlyRefreshSurfaceMaterials);
 							}
 						}
 					}
diff --git a/Engine/Source/Editor/UnrealEd/Private/Fbx/FbxMainExport.cpp b/Engine/Source/Editor/UnrealEd/Private/Fbx/FbxMainExport.cpp
--- a/Engine/Source/Editor/UnrealEd/Private/Fbx/FbxMainExport.cpp
+++ b/Engine/Source/Editor/UnrealEd/Private/Fbx/FbxMainExport.cpp
@@ -1002,7 +1002,7 @@ void FFbxExporter::ExportBSP( UModel* Model, bool bSelectedOnly )
 				FBSPExportData& Data = BrushToMeshMap.FindOrAdd( BrushActor );
 
 				FPoly Poly;
-				GEditor->polyFindMaster(Model, Node.iSurf, Poly);
+				GEditor->polyFindBrush(Model, Node.iSurf, Poly);
 
 				Data.NumVerts += Node.NumVertices;
 				Data.NumFaces += Node.NumVertices-2;
@@ -1024,7 +1024,7 @@ void FFbxExporter::ExportBSP( UModel* Model, bool bSelectedOnly )
 		if( (Surf.PolyFlags & PF_Selected) || !bSelectedOnly || (BrushActor && BrushActor->IsSelected() && Node.NumVertices >= 3) )
 		{
 			FPoly Poly;
-			GEditor->polyFindMaster( Model, Node.iSurf, Poly );
+			GEditor->polyFindBrush( Model, Node.iSurf, Poly );
 
 			FBSPExportData* ExportData = BrushToMeshMap.Find( BrushActor );
 			if( NULL == ExportData )
diff --git a/Engine/Source/Editor/UnrealEd/Private/LevelEditorViewport.cpp b/Engine/Source/Editor/UnrealEd/Private/LevelEditorViewport.cpp
--- a/Engine/Source/Editor/UnrealEd/Private/LevelEditorViewport.cpp
+++ b/Engine/Source/Editor/UnrealEd/Private/LevelEditorViewport.cpp
@@ -897,7 +897,7 @@ bool FLevelEditorViewportClient::AttemptApplyObjAsMaterialToSurface( UObject* Ob
 				Model->Surfs[SelectedSurfIndex].Material = DroppedObjAsMaterial;
 				const bool bUpdateTexCoords = false;
 				const bool bOnlyRefreshSurfaceMaterials = true;
-				GEditor->polyUpdateMaster(Model, SelectedSurfIndex, bUpdateTexCoords, bOnlyRefreshSurfaceMaterials);
+				GEditor->polyUpdateBrush(Model, SelectedSurfIndex, bUpdateTexCoords, bOnlyRefreshSurfaceMaterials);
 			}
 
 			bResult = true;
diff --git a/Engine/Source/Editor/UnrealEd/Private/LevelViewportClickHandlers.cpp b/Engine/Source/Editor/UnrealEd/Private/LevelViewportClickHandlers.cpp
--- a/Engine/Source/Editor/UnrealEd/Private/LevelViewportClickHandlers.cpp
+++ b/Engine/Source/Editor/UnrealEd/Private/LevelViewportClickHandlers.cpp
@@ -633,7 +633,7 @@ namespace LevelViewportClickHandlers
 					Model->Surfs[i].Material = SelectedMaterialInstance;
 					const bool bUpdateTexCoords = false;
 					const bool bOnlyRefreshSurfaceMaterials = true;
-					GEditor->polyUpdateMaster(Model, i, bUpdateTexCoords, bOnlyRefreshSurfaceMaterials);
+					GEditor->polyUpdateBrush(Model, i, bUpdateTexCoords, bOnlyRefreshSurfaceMaterials);
 				}
 			}
 		}
@@ -703,13 +703,13 @@ namespace LevelViewportClickHandlers
 				Surf.PolyFlags	= GSaveSurf.PolyFlags;
 				const bool bUpdateTexCoords = true;
 				const bool bOnlyRefreshSurfaceMaterials = true;
-				GEditor->polyUpdateMaster(Model, iSurf, bUpdateTexCoords, bOnlyRefreshSurfaceMaterials);
+				GEditor->polyUpdateBrush(Model, iSurf, bUpdateTexCoords, bOnlyRefreshSurfaceMaterials);
 			}
 			else
 			{
 				const bool bUpdateTexCoords = false;
 				const bool bOnlyRefreshSurfaceMaterials = true;
-				GEditor->polyUpdateMaster(Model, iSurf, bUpdateTexCoords, bOnlyRefreshSurfaceMaterials);
+				GEditor->polyUpdateBrush(Model, iSurf, bUpdateTexCoords, bOnlyRefreshSurfaceMaterials);
 			}
 		}
 		else if( Click.GetKey() == EKeys::RightMouseButton && !Click.IsControlDown() )
diff --git a/Engine/Source/Editor/UnrealEd/Private/TexAlignTools.cpp b/Engine/Source/Editor/UnrealEd/Private/TexAlignTools.cpp
--- a/Engine/Source/Editor/UnrealEd/Private/TexAlignTools.cpp
+++ b/Engine/Source/Editor/UnrealEd/Private/TexAlignTools.cpp
@@ -105,7 +105,7 @@ void UTexAligner::Align( UWorld* InWorld, ETexAlign InTexAlignType, UModel* InMo
 	{
 		FBspSurfIdx* Surf = &InitialSurfList[i];
 //		Normal = InModel->Vectors[ Surf->Surf->vNormal ];
-//		GEditor->polyFindMaster( InModel, Surf->Idx, EdPoly );
+//		GEditor->polyFindBrush( InModel, Surf->Idx, EdPoly );
 
 		bool bOK = 1;
 		/*
@@ -125,14 +125,14 @@ void UTexAligner::Align( UWorld* InWorld, ETexAlign InTexAlignType, UModel* InMo
 	for( int32 i = 0 ; i < FinalSurfList.Num() ; i++ )
 	{
 		FBspSurfIdx* Surf = &FinalSurfList[i];
-		GEditor->polyFindMaster( InModel, Surf->Idx, EdPoly );
+		GEditor->polyFindBrush( InModel, Surf->Idx, EdPoly );
 		Normal = (FVector)InModel->Vectors[ Surf->Surf->vNormal ];
 
 		AlignSurf( InTexAlignType == TEXALIGN_None ? (ETexAlign)DefTexAlign : InTexAlignType, InModel, Surf, &EdPoly, &Normal );
 
 		const bool bUpdateTexCoords = true;
 		const bool bOnlyRefreshSurfaceMaterials = true;
-		GEditor->polyUpdateMaster(InModel, Surf->Idx, bUpdateTexCoords, bOnlyRefreshSurfaceMaterials);
+		GEditor->polyUpdateBrush(InModel, Surf->Idx, bUpdateTexCoords, bOnlyRefreshSurfaceMaterials);
 	}
 
 	GEditor->RedrawLevelEditingViewports();
