BSPUtils: Fix truncation warnings

diff --git a/Engine/Source/Developer/BSPUtils/Private/BSPOps.cpp b/Engine/Source/Developer/BSPUtils/Private/BSPOps.cpp
--- a/Engine/Source/Developer/BSPUtils/Private/BSPOps.cpp
+++ b/Engine/Source/Developer/BSPUtils/Private/BSPOps.cpp
@@ -373,7 +373,7 @@ static FPoly *FindBestSplit
 		}
 		// added by Legend 1/31/1999
 		// Score optimization: minimize cuts vs. balance tree (as specified in BSP Rebuilder dialog)
-		Score = ( 100.0 - float(Balance) ) * Splits + float(Balance) * FMath::Abs( Front - Back );
+		Score = static_cast<float>(( 100.0 - float(Balance) ) * Splits + float(Balance) * FMath::Abs( Front - Back ));
 		if( Poly->PolyFlags & PF_Portal )
 		{
 			// PortalBias -- added by Legend on 4/12/2000
@@ -393,7 +393,7 @@ static FPoly *FindBestSplit
 			// should not be so high that portals cut through adjacent geometry in a way that
 			// increases complexity of the room being (typically, accidentally) cut.
 			//
-			Score -= ( 100.0 - float(Balance) ) * Splits * PortalBias; // ignore PortalBias of the split polys -- bias toward portal selection for cutting planes!
+			Score -= static_cast<float>(( 100.0 - float(Balance) ) * Splits * PortalBias); // ignore PortalBias of the split polys -- bias toward portal selection for cutting planes!
 		}
 		//UE_LOG(LogBSPOps, Log,  "  %4d: Score = %f (Front = %4d, Back = %4d, Splits = %4d, Flags = %08X)", Index, Score, Front, Back, Splits, Poly->PolyFlags ); //LEC
 
@@ -1187,7 +1187,7 @@ int32	FBSPOps::bspAddNode( UModel* Model, int32 iParent, ENodePlace NodePlace, u
 
 		// Set node properties.
 		Node.iSurf       	 = EdPoly->iLinkSurf;
-		Node.NodeFlags   	 = NodeFlags;
+		Node.NodeFlags   	 = static_cast<uint8>(NodeFlags);
 		Node.iCollisionBound = INDEX_NONE;
 		Node.Plane           = FPlane4f( EdPoly->Vertices[0], EdPoly->Normal );
 		Node.iVertPool       = Model->Verts.AddUninitialized(EdPoly->Vertices.Num());
@@ -1367,9 +1367,9 @@ int32 FBspPointsGrid::FindOrAddPoint(const FVector& Point, int32 Index, float Po
 	// Offset applied to the grid coordinates so aligned vertices (the normal case) don't overlap several grid items (taking into account the threshold)
 	const float GridOffset = 0.12345f;
 
-	const float AdjustedPointX = Point.X - GridOffset;
-	const float AdjustedPointY = Point.Y - GridOffset;
-	const float AdjustedPointZ = Point.Z - GridOffset;
+	const float AdjustedPointX = static_cast<float>(Point.X - GridOffset);
+	const float AdjustedPointY = static_cast<float>(Point.Y - GridOffset);
+	const float AdjustedPointZ = static_cast<float>(Point.Z - GridOffset);
 
 	const float GridX = AdjustedPointX * OneOverGranularity;
 	const float GridY = AdjustedPointY * OneOverGranularity;
diff --git a/Engine/Source/Developer/BSPUtils/Private/BSPUtils.cpp b/Engine/Source/Developer/BSPUtils/Private/BSPUtils.cpp
--- a/Engine/Source/Developer/BSPUtils/Private/BSPUtils.cpp
+++ b/Engine/Source/Developer/BSPUtils/Private/BSPUtils.cpp
@@ -1627,12 +1627,12 @@ int DistributePoint
 				// Create cutting plane perpendicular to both this side and the polygon's normal.
 				FVector Side            = FVector(Model->Points[VertPool[i].pVertex] - Model->Points[VertPool[j].pVertex]);
 				FVector SidePlaneNormal = Side ^ (FVector)Model->Nodes[iNode].Plane;
-				float   SizeSquared     = SidePlaneNormal.SizeSquared();
+				const float SizeSquared = static_cast<float>(SidePlaneNormal.SizeSquared());
 
 				if( SizeSquared>FMath::Square(0.001) )
 				{
 					// Points aren't coincedent.
-					Dist = ((FVector(Model->Points[pVertex] - Model->Points[VertPool[i].pVertex])) | SidePlaneNormal)/FMath::Sqrt(SizeSquared);
+					Dist = static_cast<float>((FVector(Model->Points[pVertex] - Model->Points[VertPool[i].pVertex]) | SidePlaneNormal) / FMath::Sqrt(SizeSquared));
 					if( Dist >= THRESH_OPTGEOM_COSIDAL )
 					{
 						// Point is outside polygon, can't possibly fall on a side.
@@ -1763,7 +1763,7 @@ void MergeNearPoints( UModel *Model, float Dist )
 			if( A->pVertex != B->pVertex )
 				Pool[k++] = Pool[j];
 		}
-		Node.NumVertices = k>=3 ? k : 0;
+		Node.NumVertices = static_cast<uint8>(k>=3 ? k : 0);
 		if( k < 3 )
 			Collapsed++;
 	}
@@ -2098,7 +2098,7 @@ void FBSPUtils::polySplitOverlappingEdges( TArray<FPoly>* InPolyList, TArray<FPo
 
 						FVector Dir = SrcEdge.Vertex[1] - SrcEdge.Vertex[0];
 						Dir.Normalize();
-						float Dist = FVector::Dist( SrcEdge.Vertex[1], SrcEdge.Vertex[0] );
+						const float Dist = static_cast<float>(FVector::Dist(SrcEdge.Vertex[1], SrcEdge.Vertex[0]));
 						FVector Origin = SrcEdge.Vertex[0] + (Dir * (Dist / 2.0f));
 						float Radius = Dist / 2.0f;
 
