More operator new (TArray) removal.

diff --git a/Engine/Source/Developer/AudioFormatADPCM/Private/AudioFormatADPCM.cpp b/Engine/Source/Developer/AudioFormatADPCM/Private/AudioFormatADPCM.cpp
--- a/Engine/Source/Developer/AudioFormatADPCM/Private/AudioFormatADPCM.cpp
+++ b/Engine/Source/Developer/AudioFormatADPCM/Private/AudioFormatADPCM.cpp
@@ -641,7 +641,7 @@ public:
 	// Add a new chunk and reserve ChunkSize bytes in it
 	void AddNewChunk(TArray<TArray<uint8>>& OutBuffers, int32 ChunkReserveSize) const
 	{
-		TArray<uint8>& NewBuffer = *new (OutBuffers) TArray<uint8>;
+		TArray<uint8>& NewBuffer = OutBuffers.AddDefaulted_GetRef();
 		NewBuffer.Empty(ChunkReserveSize);
 	}
 	
diff --git a/Engine/Source/Developer/AudioFormatOgg/Private/AudioFormatOgg.cpp b/Engine/Source/Developer/AudioFormatOgg/Private/AudioFormatOgg.cpp
--- a/Engine/Source/Developer/AudioFormatOgg/Private/AudioFormatOgg.cpp
+++ b/Engine/Source/Developer/AudioFormatOgg/Private/AudioFormatOgg.cpp
@@ -511,7 +511,7 @@ public:
 	// Add a new chunk and reserve ChunkSize bytes in it
 	void AddNewChunk(TArray<TArray<uint8>>& OutBuffers, int32 ChunkReserveSize) const
 	{
-		TArray<uint8>& NewBuffer = *new (OutBuffers) TArray<uint8>;
+		TArray<uint8>& NewBuffer = OutBuffers.AddDefaulted_GetRef();
 		NewBuffer.Empty(ChunkReserveSize);
 	}
 	
diff --git a/Engine/Source/Developer/AudioFormatOpus/Private/AudioFormatOpus.cpp b/Engine/Source/Developer/AudioFormatOpus/Private/AudioFormatOpus.cpp
--- a/Engine/Source/Developer/AudioFormatOpus/Private/AudioFormatOpus.cpp
+++ b/Engine/Source/Developer/AudioFormatOpus/Private/AudioFormatOpus.cpp
@@ -551,7 +551,7 @@ public:
 	 */
 	int32 AddDataChunk(TArray<TArray<uint8>>& OutBuffers, const uint8* ChunkData, int32 ChunkSize) const
 	{
-		TArray<uint8>& NewBuffer = *new (OutBuffers) TArray<uint8>;
+		TArray<uint8>& NewBuffer = OutBuffers.AddDefaulted_GetRef();
 		NewBuffer.Empty(ChunkSize);
 		NewBuffer.AddUninitialized(ChunkSize);
 		FMemory::Memcpy(NewBuffer.GetData(), ChunkData, ChunkSize);
diff --git a/Engine/Source/Developer/BSPUtils/Private/BSPOps.cpp b/Engine/Source/Developer/BSPUtils/Private/BSPOps.cpp
--- a/Engine/Source/Developer/BSPUtils/Private/BSPOps.cpp
+++ b/Engine/Source/Developer/BSPUtils/Private/BSPOps.cpp
@@ -153,10 +153,10 @@ FPoly FBSPOps::BuildInfiniteFPoly( UModel* Model, int32 iNode )
 	EdPoly.Init();
 	EdPoly.Normal      = (FVector3f)Normal;
 	EdPoly.Base        = (FVector3f)Base;
-	new(EdPoly.Vertices) FVector3f(Base + Axis1*UE_OLD_WORLD_MAX + Axis2*UE_OLD_WORLD_MAX);
-	new(EdPoly.Vertices) FVector3f(Base - Axis1*UE_OLD_WORLD_MAX + Axis2*UE_OLD_WORLD_MAX);
-	new(EdPoly.Vertices) FVector3f(Base - Axis1*UE_OLD_WORLD_MAX - Axis2*UE_OLD_WORLD_MAX);
-	new(EdPoly.Vertices) FVector3f(Base + Axis1*UE_OLD_WORLD_MAX - Axis2*UE_OLD_WORLD_MAX);
+	EdPoly.Vertices.Emplace(Base + Axis1*UE_OLD_WORLD_MAX + Axis2*UE_OLD_WORLD_MAX);
+	EdPoly.Vertices.Emplace(Base - Axis1*UE_OLD_WORLD_MAX + Axis2*UE_OLD_WORLD_MAX);
+	EdPoly.Vertices.Emplace(Base - Axis1*UE_OLD_WORLD_MAX - Axis2*UE_OLD_WORLD_MAX);
+	EdPoly.Vertices.Emplace(Base + Axis1*UE_OLD_WORLD_MAX - Axis2*UE_OLD_WORLD_MAX);
 
 	return EdPoly;
 }
@@ -971,45 +971,45 @@ void FBSPOps::bspBuildBounds( UModel* Model )
 		PolyList[i]->iBrushPoly = INDEX_NONE;
 	}
 
-	new(Polys[0].Vertices)FVector3f(-HALF_WORLD_MAX,-HALF_WORLD_MAX,HALF_WORLD_MAX);
-	new(Polys[0].Vertices)FVector3f( HALF_WORLD_MAX,-HALF_WORLD_MAX,HALF_WORLD_MAX);
-	new(Polys[0].Vertices)FVector3f( HALF_WORLD_MAX, HALF_WORLD_MAX,HALF_WORLD_MAX);
-	new(Polys[0].Vertices)FVector3f(-HALF_WORLD_MAX, HALF_WORLD_MAX,HALF_WORLD_MAX);
+	Polys[0].Vertices.Emplace(-HALF_WORLD_MAX,-HALF_WORLD_MAX,HALF_WORLD_MAX);
+	Polys[0].Vertices.Emplace( HALF_WORLD_MAX,-HALF_WORLD_MAX,HALF_WORLD_MAX);
+	Polys[0].Vertices.Emplace( HALF_WORLD_MAX, HALF_WORLD_MAX,HALF_WORLD_MAX);
+	Polys[0].Vertices.Emplace(-HALF_WORLD_MAX, HALF_WORLD_MAX,HALF_WORLD_MAX);
 	Polys[0].Normal   =FVector3f( 0.000000,  0.000000,  1.000000 );
 	Polys[0].Base     =Polys[0].Vertices[0];
 
-	new(Polys[1].Vertices)FVector3f(-HALF_WORLD_MAX, HALF_WORLD_MAX,-HALF_WORLD_MAX);
-	new(Polys[1].Vertices)FVector3f( HALF_WORLD_MAX, HALF_WORLD_MAX,-HALF_WORLD_MAX);
-	new(Polys[1].Vertices)FVector3f( HALF_WORLD_MAX,-HALF_WORLD_MAX,-HALF_WORLD_MAX);
-	new(Polys[1].Vertices)FVector3f(-HALF_WORLD_MAX,-HALF_WORLD_MAX,-HALF_WORLD_MAX);
+	Polys[1].Vertices.Emplace(-HALF_WORLD_MAX, HALF_WORLD_MAX,-HALF_WORLD_MAX);
+	Polys[1].Vertices.Emplace( HALF_WORLD_MAX, HALF_WORLD_MAX,-HALF_WORLD_MAX);
+	Polys[1].Vertices.Emplace( HALF_WORLD_MAX,-HALF_WORLD_MAX,-HALF_WORLD_MAX);
+	Polys[1].Vertices.Emplace(-HALF_WORLD_MAX,-HALF_WORLD_MAX,-HALF_WORLD_MAX);
 	Polys[1].Normal   =FVector3f( 0.000000,  0.000000, -1.000000 );
 	Polys[1].Base     =Polys[1].Vertices[0];
 
-	new(Polys[2].Vertices)FVector3f(-HALF_WORLD_MAX,HALF_WORLD_MAX,-HALF_WORLD_MAX);
-	new(Polys[2].Vertices)FVector3f(-HALF_WORLD_MAX,HALF_WORLD_MAX, HALF_WORLD_MAX);
-	new(Polys[2].Vertices)FVector3f( HALF_WORLD_MAX,HALF_WORLD_MAX, HALF_WORLD_MAX);
-	new(Polys[2].Vertices)FVector3f( HALF_WORLD_MAX,HALF_WORLD_MAX,-HALF_WORLD_MAX);
+	Polys[2].Vertices.Emplace(-HALF_WORLD_MAX,HALF_WORLD_MAX,-HALF_WORLD_MAX);
+	Polys[2].Vertices.Emplace(-HALF_WORLD_MAX,HALF_WORLD_MAX, HALF_WORLD_MAX);
+	Polys[2].Vertices.Emplace( HALF_WORLD_MAX,HALF_WORLD_MAX, HALF_WORLD_MAX);
+	Polys[2].Vertices.Emplace( HALF_WORLD_MAX,HALF_WORLD_MAX,-HALF_WORLD_MAX);
 	Polys[2].Normal   =FVector3f( 0.000000,  1.000000,  0.000000 );
 	Polys[2].Base     =Polys[2].Vertices[0];
 
-	new(Polys[3].Vertices)FVector3f( HALF_WORLD_MAX,-HALF_WORLD_MAX,-HALF_WORLD_MAX);
-	new(Polys[3].Vertices)FVector3f( HALF_WORLD_MAX,-HALF_WORLD_MAX, HALF_WORLD_MAX);
-	new(Polys[3].Vertices)FVector3f(-HALF_WORLD_MAX,-HALF_WORLD_MAX, HALF_WORLD_MAX);
-	new(Polys[3].Vertices)FVector3f(-HALF_WORLD_MAX,-HALF_WORLD_MAX,-HALF_WORLD_MAX);
+	Polys[3].Vertices.Emplace( HALF_WORLD_MAX,-HALF_WORLD_MAX,-HALF_WORLD_MAX);
+	Polys[3].Vertices.Emplace( HALF_WORLD_MAX,-HALF_WORLD_MAX, HALF_WORLD_MAX);
+	Polys[3].Vertices.Emplace(-HALF_WORLD_MAX,-HALF_WORLD_MAX, HALF_WORLD_MAX);
+	Polys[3].Vertices.Emplace(-HALF_WORLD_MAX,-HALF_WORLD_MAX,-HALF_WORLD_MAX);
 	Polys[3].Normal   =FVector3f( 0.000000, -1.000000,  0.000000 );
 	Polys[3].Base     =Polys[3].Vertices[0];
 
-	new(Polys[4].Vertices)FVector3f(HALF_WORLD_MAX, HALF_WORLD_MAX,-HALF_WORLD_MAX);
-	new(Polys[4].Vertices)FVector3f(HALF_WORLD_MAX, HALF_WORLD_MAX, HALF_WORLD_MAX);
-	new(Polys[4].Vertices)FVector3f(HALF_WORLD_MAX,-HALF_WORLD_MAX, HALF_WORLD_MAX);
-	new(Polys[4].Vertices)FVector3f(HALF_WORLD_MAX,-HALF_WORLD_MAX,-HALF_WORLD_MAX);
+	Polys[4].Vertices.Emplace(HALF_WORLD_MAX, HALF_WORLD_MAX,-HALF_WORLD_MAX);
+	Polys[4].Vertices.Emplace(HALF_WORLD_MAX, HALF_WORLD_MAX, HALF_WORLD_MAX);
+	Polys[4].Vertices.Emplace(HALF_WORLD_MAX,-HALF_WORLD_MAX, HALF_WORLD_MAX);
+	Polys[4].Vertices.Emplace(HALF_WORLD_MAX,-HALF_WORLD_MAX,-HALF_WORLD_MAX);
 	Polys[4].Normal   =FVector3f( 1.000000,  0.000000,  0.000000 );
 	Polys[4].Base     =Polys[4].Vertices[0];
 
-	new(Polys[5].Vertices)FVector3f(-HALF_WORLD_MAX,-HALF_WORLD_MAX,-HALF_WORLD_MAX);
-	new(Polys[5].Vertices)FVector3f(-HALF_WORLD_MAX,-HALF_WORLD_MAX, HALF_WORLD_MAX);
-	new(Polys[5].Vertices)FVector3f(-HALF_WORLD_MAX, HALF_WORLD_MAX, HALF_WORLD_MAX);
-	new(Polys[5].Vertices)FVector3f(-HALF_WORLD_MAX, HALF_WORLD_MAX,-HALF_WORLD_MAX);
+	Polys[5].Vertices.Emplace(-HALF_WORLD_MAX,-HALF_WORLD_MAX,-HALF_WORLD_MAX);
+	Polys[5].Vertices.Emplace(-HALF_WORLD_MAX,-HALF_WORLD_MAX, HALF_WORLD_MAX);
+	Polys[5].Vertices.Emplace(-HALF_WORLD_MAX, HALF_WORLD_MAX, HALF_WORLD_MAX);
+	Polys[5].Vertices.Emplace(-HALF_WORLD_MAX, HALF_WORLD_MAX,-HALF_WORLD_MAX);
 	Polys[5].Normal   =FVector3f(-1.000000,  0.000000,  0.000000 );
 	Polys[5].Base     =Polys[5].Vertices[0];
 	// Empty hulls.
diff --git a/Engine/Source/Developer/BSPUtils/Private/BSPUtils.cpp b/Engine/Source/Developer/BSPUtils/Private/BSPUtils.cpp
--- a/Engine/Source/Developer/BSPUtils/Private/BSPUtils.cpp
+++ b/Engine/Source/Developer/BSPUtils/Private/BSPUtils.cpp
@@ -170,7 +170,7 @@ int TryToMerge( FPoly *Poly1, FPoly *Poly2 )
 	int32 Vertex = End1;
 	for( int32 i=0; i<Poly1->Vertices.Num(); i++ )
 	{
-		new(NewPoly.Vertices) FVector3f(Poly1->Vertices[Vertex]);
+		NewPoly.Vertices.Add(Poly1->Vertices[Vertex]);
 		if( ++Vertex >= Poly1->Vertices.Num() )
 			Vertex=0;
 	}
@@ -179,7 +179,7 @@ int TryToMerge( FPoly *Poly1, FPoly *Poly2 )
 	{
 		if( ++Vertex >= Poly2->Vertices.Num() )
 			Vertex=0;
-		new(NewPoly.Vertices) FVector3f(Poly2->Vertices[Vertex]);
+		NewPoly.Vertices.Add(Poly2->Vertices[Vertex]);
 	}
 
 	// Remove colinear vertices and check convexity.
@@ -230,7 +230,7 @@ void MakeEdPolys( UModel* Model, int32 iNode, TArray<FPoly>* DestArray )
 	FPoly Temp;
 	if( FBSPUtils::bspNodeToFPoly(Model,iNode,&Temp) >= 3 )
 	{
-		new(*DestArray)FPoly(Temp);
+		DestArray->Add(Temp);
 	}
 
 	if( Node->iFront!=INDEX_NONE )
@@ -553,7 +553,7 @@ void IntersectBrushWithWorldFunc( UModel* Model, int32 iNode, FPoly *EdPoly,
 		case F_INSIDE:
 		case F_COPLANAR_INSIDE:
 			if( EdPoly->Fix()>=3 )
-				new(GModel->Polys->Element)FPoly(*EdPoly);
+				GModel->Polys->Element.Add(*EdPoly);
 			break;
 	}
 }
@@ -572,7 +572,7 @@ void IntersectWorldWithBrushFunc( UModel *Model, int32 iNode, FPoly *EdPoly,
 		case F_COPLANAR_INSIDE:
 		case F_COSPATIAL_FACING_OUT:
 			if( EdPoly->Fix() >= 3 )
-				new(GModel->Polys->Element)FPoly(*EdPoly);
+				GModel->Polys->Element.Add(*EdPoly);
 			break;
 	}
 }
@@ -591,7 +591,7 @@ void DeIntersectBrushWithWorldFunc( UModel* Model, int32 iNode, FPoly* EdPoly,
 		case F_OUTSIDE:
 		case F_COPLANAR_OUTSIDE:
 			if( EdPoly->Fix()>=3 )
-				new(GModel->Polys->Element)FPoly(*EdPoly);
+				GModel->Polys->Element.Add(*EdPoly);
 			break;
 	}
 }
@@ -612,7 +612,7 @@ void DeIntersectWorldWithBrushFunc( UModel* Model, int32 iNode, FPoly* EdPoly,
 			if( EdPoly->Fix() >= 3 )
 	        {
 				EdPoly->Reverse();
-				new(GModel->Polys->Element)FPoly(*EdPoly);
+				GModel->Polys->Element.Add(*EdPoly);
 				EdPoly->Reverse();
 			}
 			break;
@@ -964,7 +964,7 @@ int FBSPUtils::bspNodeToFPoly
 
 	for(int32 VertexIndex = 0;VertexIndex < Node.NumVertices;VertexIndex++)
 	{
-		new(EdPoly->Vertices) FVector3f(Model->Points[VertPool[VertexIndex].pVertex]);
+		EdPoly->Vertices.Add(Model->Points[VertPool[VertexIndex].pVertex]);
 	}
 
 	if(EdPoly->Vertices.Num() < 3)
@@ -1222,7 +1222,7 @@ int FBSPUtils::bspBrushCSG
 		}
 
 		// Add poly to the temp model.
-		new(TempModel->Polys->Element)FPoly( DestEdPoly );
+		TempModel->Polys->Element.Add( DestEdPoly );
 	}
 	if( ReallyBig ) GWarn->StatusUpdate( 0, 0, NSLOCTEXT("UnrealEd", "FilteringBrush", "Filtering brush") );
 
@@ -2049,14 +2049,14 @@ void FBSPUtils::polyGetLinkedPolys
 	if( InPoly->iLink == INDEX_NONE )
 	{
 		// If this poly has no links, just stick the one poly in the final list.
-		new(*InPolyList)FPoly( *InPoly );
+		InPolyList->Add( *InPoly );
 	}
 	else
 	{
 		// Find all polys that match the source polys link value.
 		for( int32 poly = 0 ; poly < InBrush->Brush->Polys->Element.Num() ; poly++ )
 			if( InBrush->Brush->Polys->Element[poly].iLink == InPoly->iLink )
-				new(*InPolyList)FPoly( InBrush->Brush->Polys->Element[poly] );
+				InPolyList->Add( InBrush->Brush->Polys->Element[poly] );
 	}
 }
 
@@ -2111,7 +2111,7 @@ void FBSPUtils::polySplitOverlappingEdges( TArray<FPoly>* InPolyList, TArray<FPo
 			}
 		}
 
-		new(*InResult)FPoly( NewPoly );
+		InResult->Add( NewPoly );
 	}
 
 }
@@ -2133,7 +2133,7 @@ void FBSPUtils::polyGetOuterEdgeList
 	{
 		FPoly* Poly = &NewPolyList[poly];
 		for( int32 vtx = 0 ; vtx < Poly->Vertices.Num() ; vtx++ )
-			new( TempEdges )FEdge( (FVector)Poly->Vertices[vtx], (FVector)Poly->Vertices[ vtx+1 < Poly->Vertices.Num() ? vtx+1 : 0] );
+			TempEdges.Emplace( (FVector)Poly->Vertices[vtx], (FVector)Poly->Vertices[ vtx+1 < Poly->Vertices.Num() ? vtx+1 : 0] );
 	}
 
 	// Add all the unique edges into the final edge list.
@@ -2151,14 +2151,14 @@ void FBSPUtils::polyGetOuterEdgeList
 		}
 
 		if( EdgeCount == 1 )
-			new( FinalEdges )FEdge( *TestEdge );
+			FinalEdges.Add( *TestEdge );
 	}
 
 	// Reorder all the edges so that they line up, end to end.
 	InEdgeList->Empty();
 	if( !FinalEdges.Num() ) return;
 
-	new( *InEdgeList )FEdge( FinalEdges[0] );
+	InEdgeList->Add( FinalEdges[0] );
 	FVector Comp = FinalEdges[0].Vertex[1];
 	FinalEdges.RemoveAt(0);
 
@@ -2173,7 +2173,7 @@ void FBSPUtils::polyGetOuterEdgeList
 
 		if( FinalEdges[x].Vertex[0] == Comp )
 		{
-			new( *InEdgeList )FEdge( FinalEdges[x] );
+			InEdgeList->Add( FinalEdges[x] );
 			Comp = FinalEdges[x].Vertex[1];
 			FinalEdges.RemoveAt(x);
 			x = -1;
diff --git a/Engine/Source/Editor/UnrealEd/Private/Commandlets/ContentCommandlets.cpp b/Engine/Source/Editor/UnrealEd/Private/Commandlets/ContentCommandlets.cpp
--- a/Engine/Source/Editor/UnrealEd/Private/Commandlets/ContentCommandlets.cpp
+++ b/Engine/Source/Editor/UnrealEd/Private/Commandlets/ContentCommandlets.cpp
@@ -2970,7 +2970,7 @@ int32 UWrangleContentCommandlet::Main( const FString& Params )
 				if (PackageObjs == NULL)
 				{
 					UE_LOG(LogContentCommandlet, Warning, TEXT("No objects in %s were referenced..."), *PackagePath.GetDebugName());
-					new(UnnecessaryPublicObjects) FUnreferencedObject(PackageName,
+					UnnecessaryPublicObjects.Emplace(PackageName,
 						TEXT("ENTIRE PACKAGE"), IPackageResourceManager::Get().FileSize(PackagePath));
 
 					// all objects in this package are unnecessary
@@ -3017,7 +3017,7 @@ int32 UWrangleContentCommandlet::Main( const FString& Params )
 					if ((Export.ObjectFlags & RF_Public) != 0 && !bAreAllObjectsUnnecessary)
 					{
 						// if so, then add it to list of unused pcreateexportublic items
-						new(UnnecessaryPublicObjects) FUnreferencedObject(PackageName, ExportName, Export.SerialSize);
+						UnnecessaryPublicObjects.Emplace(PackageName, ExportName, Export.SerialSize);
 					}
 
 					// look for existing entry
diff --git a/Engine/Source/Editor/UnrealEd/Private/EditorBrushBuilder.cpp b/Engine/Source/Editor/UnrealEd/Private/EditorBrushBuilder.cpp
--- a/Engine/Source/Editor/UnrealEd/Private/EditorBrushBuilder.cpp
+++ b/Engine/Source/Editor/UnrealEd/Private/EditorBrushBuilder.cpp
@@ -139,11 +139,11 @@ bool UEditorBrushBuilder::EndBrush( UWorld* InWorld, ABrush* InBrush )
 
 		for( int32 j=0; j<It->VertexIndices.Num(); j++ )
 		{
-			new(Poly.Vertices) FVector3f(Vertices[It->VertexIndices[j]]);
+			Poly.Vertices.Emplace(Vertices[It->VertexIndices[j]]);
 		}
 		if( Poly.Finalize( BuilderBrush, 1 ) == 0 )
 		{
-			new(Brush->Polys->Element)FPoly(Poly);
+			Brush->Polys->Element.Add(Poly);
 		}
 	}
 
@@ -199,7 +199,7 @@ bool UEditorBrushBuilder::BadParameters(const FText& Msg)
 int32 UEditorBrushBuilder::Vertexv(FVector V)
 {
 	int32 Result = Vertices.Num();
-	new(Vertices)FVector(V);
+	Vertices.Add(V);
 
 	return Result;
 }
@@ -207,44 +207,44 @@ int32 UEditorBrushBuilder::Vertexv(FVector V)
 int32 UEditorBrushBuilder::Vertex3f(float X, float Y, float Z)
 {
 	int32 Result = Vertices.Num();
-	new(Vertices)FVector(X,Y,Z);
+	Vertices.Emplace(X,Y,Z);
 	return Result;
 }
 
 void UEditorBrushBuilder::Poly3i(int32 Direction, int32 i, int32 j, int32 k, FName ItemName, bool bIsTwoSidedNonSolid )
 {
-	new(Polys)FBuilderPoly;
-	Polys.Last().Direction=Direction;
-	Polys.Last().ItemName=ItemName;
-	new(Polys.Last().VertexIndices)int32(i);
-	new(Polys.Last().VertexIndices)int32(j);
-	new(Polys.Last().VertexIndices)int32(k);
-	Polys.Last().PolyFlags = PF_DefaultFlags | (bIsTwoSidedNonSolid ? (PF_TwoSided|PF_NotSolid) : 0);
+	FBuilderPoly& Poly = Polys.AddDefaulted_GetRef();
+	Poly.Direction=Direction;
+	Poly.ItemName=ItemName;
+	Poly.VertexIndices.Add(i);
+	Poly.VertexIndices.Add(j);
+	Poly.VertexIndices.Add(k);
+	Poly.PolyFlags = PF_DefaultFlags | (bIsTwoSidedNonSolid ? (PF_TwoSided|PF_NotSolid) : 0);
 }
 
 void UEditorBrushBuilder::Poly4i(int32 Direction, int32 i, int32 j, int32 k, int32 l, FName ItemName, bool bIsTwoSidedNonSolid )
 {
-	new(Polys)FBuilderPoly;
-	Polys.Last().Direction=Direction;
-	Polys.Last().ItemName=ItemName;
-	new(Polys.Last().VertexIndices)int32(i);
-	new(Polys.Last().VertexIndices)int32(j);
-	new(Polys.Last().VertexIndices)int32(k);
-	new(Polys.Last().VertexIndices)int32(l);
-	Polys.Last().PolyFlags = PF_DefaultFlags | (bIsTwoSidedNonSolid ? (PF_TwoSided|PF_NotSolid) : 0);
+	FBuilderPoly& Poly = Polys.AddDefaulted_GetRef();
+	Poly.Direction=Direction;
+	Poly.ItemName=ItemName;
+	Poly.VertexIndices.Add(i);
+	Poly.VertexIndices.Add(j);
+	Poly.VertexIndices.Add(k);
+	Poly.VertexIndices.Add(l);
+	Poly.PolyFlags = PF_DefaultFlags | (bIsTwoSidedNonSolid ? (PF_TwoSided|PF_NotSolid) : 0);
 }
 
 void UEditorBrushBuilder::PolyBegin(int32 Direction, FName ItemName)
 {
-	new(Polys)FBuilderPoly;
-	Polys.Last().ItemName=ItemName;
-	Polys.Last().Direction = Direction;
-	Polys.Last().PolyFlags = PF_DefaultFlags;
+	FBuilderPoly& Poly = Polys.AddDefaulted_GetRef();
+	Poly.ItemName=ItemName;
+	Poly.Direction = Direction;
+	Poly.PolyFlags = PF_DefaultFlags;
 }
 
 void UEditorBrushBuilder::Polyi(int32 i)
 {
-	new(Polys.Last().VertexIndices)int32(i);
+	Polys.Last().VertexIndices.Add(i);
 }
 
 void UEditorBrushBuilder::PolyEnd()
